Virtual Scouting: Static analysis fix.

diff --git a/Engine/Plugins/Experimental/VirtualScouting/Source/VirtualScoutingOpenXR/Private/VirtualScoutingOpenXR.h b/Engine/Plugins/Experimental/VirtualScouting/Source/VirtualScoutingOpenXR/Private/VirtualScoutingOpenXR.h
--- a/Engine/Plugins/Experimental/VirtualScouting/Source/VirtualScoutingOpenXR/Private/VirtualScoutingOpenXR.h
+++ b/Engine/Plugins/Experimental/VirtualScouting/Source/VirtualScoutingOpenXR/Private/VirtualScoutingOpenXR.h
@@ -1,5 +1,7 @@
 // Copyright Epic Games, Inc. All Rights Reserved.
 
+#pragma once
+
 #if WITH_EDITOR
 
 #include "IOpenXRExtensionPlugin.h"
