VirtualScouting / OpenXRInput: UnrealGame NoPCH build fixes

diff --git a/Engine/Plugins/Experimental/VirtualScouting/Source/VirtualScoutingOpenXR/Private/VirtualScoutingOpenXR.cpp b/Engine/Plugins/Experimental/VirtualScouting/Source/VirtualScoutingOpenXR/Private/VirtualScoutingOpenXR.cpp
--- a/Engine/Plugins/Experimental/VirtualScouting/Source/VirtualScoutingOpenXR/Private/VirtualScoutingOpenXR.cpp
+++ b/Engine/Plugins/Experimental/VirtualScouting/Source/VirtualScoutingOpenXR/Private/VirtualScoutingOpenXR.cpp
@@ -2,9 +2,11 @@
 
 #include "VirtualScoutingOpenXR.h"
 
+#include "Logging/LogMacros.h"
+#include "Modules/ModuleManager.h"
+
 #if WITH_EDITOR
 
-#include "Logging/LogMacros.h"
 #include "IVREditorModule.h"
 #include "ViewportWorldInteraction.h"
 #include "VREditorMode.h"
diff --git a/Engine/Plugins/Runtime/OpenXR/Source/OpenXRInput/Private/OpenXRInput.h b/Engine/Plugins/Runtime/OpenXR/Source/OpenXRInput/Private/OpenXRInput.h
--- a/Engine/Plugins/Runtime/OpenXR/Source/OpenXRInput/Private/OpenXRInput.h
+++ b/Engine/Plugins/Runtime/OpenXR/Source/OpenXRInput/Private/OpenXRInput.h
@@ -11,6 +11,7 @@
 class FOpenXRHMD;
 struct FInputActionKeyMapping;
 struct FInputAxisKeyMapping;
+struct FKey;
 
 // On some platforms the XrPath type becomes ambiguous for overloading
 FORCEINLINE uint32 GetTypeHash(const TPair<XrPath, XrPath>& Pair)
