Virtual Scouting: GetHMDDeviceType backward compatibility under OpenXR.

Fixes several bugs sharing the same root cause, where code/blueprints were attempting to branch on UVREditorMode::GetHMDDeviceType() or UVREditorInteractor::GetHMDDeviceType(). These calls previously returned the result of IXRTrackingSystem::GetSystemName(), and the call sites were only set up to handle "OculusHMD" or "SteamVR," and were now seeing "OpenXR". Instead, we now return one of those legacy plugin names, depending on the current OpenXR hand controller interaction profile. Longer term, we will want to revisit each of these decision points individually, and take into account support for future unknown XR devices.

An alternative subclass of UVREditorMode can now be configured in the editor settings alongside the interactor and teleporter classes, and a new UVirtualScoutingMode subclass splits the mode entry into two phases. This is necessary because OpenXR may not return the correct interaction profile for several frames after the OpenXR session / stereo rendering has started, and we need to defer creation of the interactors, etc.

diff --git a/Engine/Config/BaseEditorSettings.ini b/Engine/Config/BaseEditorSettings.ini
--- a/Engine/Config/BaseEditorSettings.ini
+++ b/Engine/Config/BaseEditorSettings.ini
@@ -59,6 +59,7 @@ GizmoScale=0.800000
 DoubleClickTime=0.250000
 TriggerPressedThreshold_Vive=0.330000
 TriggerPressedThreshold_Rift=0.500000
+ModeClass=Class'/Script/VREditor.VREditorMode'
 InteractorClass=Class'/Script/VREditor.VREditorInteractor'
 TeleporterClass=Class'/Script/VREditor.VREditorTeleporter'
 
diff --git a/Engine/Plugins/Experimental/VirtualScouting/Source/VirtualScoutingEditor/Private/VirtualScoutingMode.cpp b/Engine/Plugins/Experimental/VirtualScouting/Source/VirtualScoutingEditor/Private/VirtualScoutingMode.cpp
new file mode 100644
--- /dev/null
+++ b/Engine/Plugins/Experimental/VirtualScouting/Source/VirtualScoutingEditor/Private/VirtualScoutingMode.cpp
@@ -0,0 +1,77 @@
+// Copyright Epic Games, Inc. All Rights Reserved.
+
+#include "VirtualScoutingMode.h"
+#include "VirtualScoutingOpenXR.h"
+#include "VirtualScoutingOpenXRModule.h"
+
+#include "IXRTrackingSystem.h"
+#include "Logging/LogMacros.h"
+#include "Modules/ModuleManager.h"
+
+
+#define LOCTEXT_NAMESPACE "VirtualScouting"
+
+
+DECLARE_LOG_CATEGORY_EXTERN(LogVirtualScouting, Log, All);
+DEFINE_LOG_CATEGORY(LogVirtualScouting);
+
+
+class FVirtualScoutingEditorModule : public IModuleInterface
+{
+private:
+	virtual void StartupModule() override
+	{
+	}
+
+	virtual void ShutdownModule() override
+	{
+	}
+};
+
+
+IMPLEMENT_MODULE(FVirtualScoutingEditorModule, VirtualScoutingEditor);
+
+
+void UVirtualScoutingMode::Enter()
+{
+	if (!GEngine->XRSystem.IsValid() || GEngine->XRSystem->GetSystemName() != FName(TEXT("OpenXR")))
+	{
+		Super::Enter();
+		return;
+	}
+
+	TSharedPtr<FVirtualScoutingOpenXRExtension> XrExt = FVirtualScoutingOpenXRModule::Get().GetOpenXRExt();
+	if (!XrExt)
+	{
+		UE_LOG(LogVirtualScouting, Error, TEXT("OpenXR extension plugin invalid"));
+		return;
+	}
+
+	// This causes FOpenXRInput to rebuild and reattach actions.
+	FEditorDelegates::OnActionAxisMappingsChanged.Broadcast();
+
+	// Split the mode entry into two phases. This is necessary because we have to poll OpenXR for
+	// the active interaction profile and translate it into a legacy plugin name, but OpenXR may
+	// not return the correct interaction profile for several frames after the OpenXR session
+	// (stereo rendering) has started, and we need to defer creation of the interactors, etc.
+	BeginEntry();
+
+	XrExt->GetHmdDeviceTypeFuture() = XrExt->GetHmdDeviceTypeFuture().Next([this](FName DeviceType)
+	{
+		if (DeviceType != NAME_None)
+		{
+			SetHMDDeviceTypeOverride(DeviceType);
+		}
+		else
+		{
+			UE_LOG(LogVirtualScouting, Error, TEXT("Unable to map legacy HMD device type"));
+		}
+
+		SetupSubsystems();
+		FinishEntry();
+
+		return DeviceType;
+	});
+}
+
+#undef LOCTEXT_NAMESPACE
diff --git a/Engine/Plugins/Experimental/VirtualScouting/Source/VirtualScoutingEditor/Private/VirtualScoutingMode.h b/Engine/Plugins/Experimental/VirtualScouting/Source/VirtualScoutingEditor/Private/VirtualScoutingMode.h
new file mode 100644
--- /dev/null
+++ b/Engine/Plugins/Experimental/VirtualScouting/Source/VirtualScoutingEditor/Private/VirtualScoutingMode.h
@@ -0,0 +1,18 @@
+// Copyright Epic Games, Inc. All Rights Reserved.
+
+#pragma once
+
+#include "VREditorMode.h"
+#include "VirtualScoutingMode.generated.h"
+
+
+UCLASS()
+class UVirtualScoutingMode : public UVREditorMode
+{
+	GENERATED_BODY()
+
+public:
+	//~ Begin UVREditorMode interface
+	virtual void Enter() override;
+	//~ End UVREditorMode interface
+};
diff --git a/Engine/Plugins/Experimental/VirtualScouting/Source/VirtualScoutingEditor/VirtualScoutingEditor.Build.cs b/Engine/Plugins/Experimental/VirtualScouting/Source/VirtualScoutingEditor/VirtualScoutingEditor.Build.cs
new file mode 100644
--- /dev/null
+++ b/Engine/Plugins/Experimental/VirtualScouting/Source/VirtualScoutingEditor/VirtualScoutingEditor.Build.cs
@@ -0,0 +1,34 @@
+// Copyright Epic Games, Inc. All Rights Reserved.
+
+using UnrealBuildTool;
+
+public class VirtualScoutingEditor : ModuleRules
+{
+	public VirtualScoutingEditor(ReadOnlyTargetRules Target) : base(Target)
+	{
+		PCHUsage = ModuleRules.PCHUsageMode.UseExplicitOrSharedPCHs;
+
+		PrivateIncludePaths.AddRange(
+			new string[]
+			{
+				"VirtualScoutingOpenXR/Private",
+			}
+		);
+
+		PublicDependencyModuleNames.AddRange(new string[]
+			{
+				"Core",
+				"CoreUObject",
+				"Engine",
+			}
+		);
+
+		PrivateDependencyModuleNames.AddRange(new string[]
+			{
+				"UnrealEd",
+				"VREditor",
+				"VirtualScoutingOpenXR",
+			}
+		);
+	}
+}
diff --git a/Engine/Plugins/Experimental/VirtualScouting/Source/VirtualScoutingOpenXR/Private/VirtualScoutingOpenXR.cpp b/Engine/Plugins/Experimental/VirtualScouting/Source/VirtualScoutingOpenXR/Private/VirtualScoutingOpenXR.cpp
--- a/Engine/Plugins/Experimental/VirtualScouting/Source/VirtualScoutingOpenXR/Private/VirtualScoutingOpenXR.cpp
+++ b/Engine/Plugins/Experimental/VirtualScouting/Source/VirtualScoutingOpenXR/Private/VirtualScoutingOpenXR.cpp
@@ -1,6 +1,7 @@
 // Copyright Epic Games, Inc. All Rights Reserved.
 
 #include "VirtualScoutingOpenXR.h"
+#include "VirtualScoutingOpenXRModule.h"
 
 #include "Logging/LogMacros.h"
 #include "Modules/ModuleManager.h"
@@ -17,8 +18,10 @@
 #define LOCTEXT_NAMESPACE "VirtualScouting"
 
 
-DECLARE_LOG_CATEGORY_EXTERN(LogVPOpenXRDebug, VeryVerbose, All);
-DEFINE_LOG_CATEGORY(LogVPOpenXRDebug);
+DECLARE_LOG_CATEGORY_EXTERN(LogVirtualScoutingOpenXR, Log, All);
+DECLARE_LOG_CATEGORY_EXTERN(LogVirtualScoutingOpenXRDebug, VeryVerbose, All);
+DEFINE_LOG_CATEGORY(LogVirtualScoutingOpenXR);
+DEFINE_LOG_CATEGORY(LogVirtualScoutingOpenXRDebug);
 
 
 static TAutoConsoleVariable<int32> CVarOpenXRDebugLogging(
@@ -28,22 +31,16 @@ static TAutoConsoleVariable<int32> CVarOpenXRDebugLogging(
 	ECVF_Default);
 
 
-class FVirtualScoutingOpenXRModule : public IModuleInterface
-{
-private:
-	virtual void StartupModule() override
-	{
-		OpenXRExt = MakeUnique<FVirtualScoutingOpenXRExtension>();
-	}
 
-	virtual void ShutdownModule() override
-	{
-		OpenXRExt.Reset();
-	}
+void FVirtualScoutingOpenXRModule::StartupModule()
+{
+	OpenXRExt = MakeShared<FVirtualScoutingOpenXRExtension>();
+}
 
-private:
-	TUniquePtr<FVirtualScoutingOpenXRExtension> OpenXRExt;
-};
+void FVirtualScoutingOpenXRModule::ShutdownModule()
+{
+	OpenXRExt.Reset();
+}
 
 
 IMPLEMENT_MODULE(FVirtualScoutingOpenXRModule, VirtualScoutingOpenXR);
@@ -56,8 +53,9 @@ FVirtualScoutingOpenXRExtension::FVirtualScoutingOpenXRExtension()
 	InitCompleteDelegate = FCoreDelegates::OnFEngineLoopInitComplete.AddLambda(
 		[this]()
 		{
-			IVREditorModule::Get().OnVREditingModeEnter().AddRaw(this, &FVirtualScoutingOpenXRExtension::OnVREditingModeEnter);
-			IVREditorModule::Get().OnVREditingModeExit().AddRaw(this, &FVirtualScoutingOpenXRExtension::OnVREditingModeExit);
+			IVREditorModule& VrEditor = IVREditorModule::Get();
+			VrEditor.OnVREditingModeEnter().AddRaw(this, &FVirtualScoutingOpenXRExtension::OnVREditingModeEnter);
+			VrEditor.OnVREditingModeExit().AddRaw(this, &FVirtualScoutingOpenXRExtension::OnVREditingModeExit);
 
 			// Must happen last; this implicitly deallocates this lambda's captures.
 			FCoreDelegates::OnFEngineLoopInitComplete.Remove(InitCompleteDelegate);
@@ -68,6 +66,11 @@ FVirtualScoutingOpenXRExtension::FVirtualScoutingOpenXRExtension()
 
 FVirtualScoutingOpenXRExtension::~FVirtualScoutingOpenXRExtension()
 {
+	if (!DeviceTypeFuture.IsReady())
+	{
+		DeviceTypePromise.EmplaceValue(NAME_None);
+	}
+
 #if 0 // FIXME?: Too late to use Instance here, and don't see any suitable extension interface methods.
 	// Might also be OK not to explicitly clean this up.
 	if (Messenger != XR_NULL_HANDLE)
@@ -103,6 +106,19 @@ bool FVirtualScoutingOpenXRExtension::GetOptionalExtensions(TArray<const ANSICHA
 }
 
 
+void FVirtualScoutingOpenXRExtension::OnEvent(XrSession InSession, const XrEventDataBaseHeader* InHeader)
+{
+	switch (InHeader->type)
+	{
+		case XR_TYPE_EVENT_DATA_INTERACTION_PROFILE_CHANGED:
+			// TODO: Correctly handle interaction profile changing mid-session
+			ensure(InSession == Session);
+			TryFulfillDeviceTypePromise();
+			break;
+	}
+}
+
+
 void FVirtualScoutingOpenXRExtension::PostCreateInstance(XrInstance InInstance)
 {
 	Instance = InInstance;
@@ -127,13 +143,22 @@ void FVirtualScoutingOpenXRExtension::PostCreateInstance(XrInstance InInstance)
 
 			if (XR_ENSURE(PfnXrCreateDebugUtilsMessengerEXT(Instance, &DebugMessengerCreateInfo, &Messenger)))
 			{
-				UE_LOG(LogVPOpenXRDebug, Log, TEXT("XR_EXT_debug_utils messenger ACTIVE"));
+				UE_LOG(LogVirtualScoutingOpenXRDebug, Log, TEXT("XR_EXT_debug_utils messenger ACTIVE"));
 				return;
 			}
 		}
 	}
 
-	UE_LOG(LogVPOpenXRDebug, Log, TEXT("XR_EXT_debug_utils messenger DISABLED"));
+	UE_LOG(LogVirtualScoutingOpenXRDebug, Log, TEXT("XR_EXT_debug_utils messenger DISABLED"));
+}
+
+
+void FVirtualScoutingOpenXRExtension::PostCreateSession(XrSession InSession)
+{
+	Session = InSession;
+
+	DeviceTypePromise = {};
+	DeviceTypeFuture = DeviceTypePromise.GetFuture();
 }
 
 
@@ -146,13 +171,13 @@ void FVirtualScoutingOpenXRExtension::AddActions(XrInstance InInstance, FCreateA
 		ActionSet = XR_NULL_HANDLE;
 	}
 
-	UClass* InteractorClass = GetDefault<UVRModeSettings>()->InteractorClass.LoadSynchronous();
+	UClass* InteractorClass = GetDefault<UVRModeSettings>()->InteractorClass.Get();
 	if (!InteractorClass)
 	{
 		return;
 	}
 
-	const UVREditorInteractor* InteractorCDO = Cast<const UVREditorInteractor>(InteractorClass->GetDefaultObject());
+	const UVREditorInteractor* InteractorCDO = InteractorClass->GetDefaultObject<const UVREditorInteractor>();
 	if (!InteractorCDO)
 	{
 		return;
@@ -201,6 +226,76 @@ void FVirtualScoutingOpenXRExtension::GetActiveActionSetsForSync(TArray<XrActive
 }
 
 
+void FVirtualScoutingOpenXRExtension::PostSyncActions(XrSession InSession)
+{
+	TryFulfillDeviceTypePromise();
+}
+
+void FVirtualScoutingOpenXRExtension::TryFulfillDeviceTypePromise()
+{
+	if (!DeviceTypeFuture.IsReady())
+	{
+		TOptional<FName> MaybeType = TryGetHmdDeviceType();
+		if (MaybeType)
+		{
+			DeviceTypePromise.EmplaceValue(MaybeType.GetValue());
+			return;
+		}
+	}
+}
+
+
+TOptional<FName> FVirtualScoutingOpenXRExtension::TryGetHmdDeviceType()
+{
+	UVREditorMode* VrMode = IVREditorModule::Get().GetVRMode();
+	if (!ensure(VrMode) || !ensure(Session != XR_NULL_HANDLE) || !ensure(ActionSet != XR_NULL_HANDLE))
+	{
+		return TOptional<FName>();
+	}
+
+	XrPath LeftHand = XR_NULL_PATH, RightHand = XR_NULL_PATH;
+	check(XR_SUCCEEDED(xrStringToPath(Instance, "/user/hand/left", &LeftHand)));
+	check(XR_SUCCEEDED(xrStringToPath(Instance, "/user/hand/right", &RightHand)));
+
+	XrInteractionProfileState LeftHandState{ XR_TYPE_INTERACTION_PROFILE_STATE },
+	                          RightHandState{ XR_TYPE_INTERACTION_PROFILE_STATE };
+
+	XR_ENSURE(xrGetCurrentInteractionProfile(Session, LeftHand, &LeftHandState));
+	XR_ENSURE(xrGetCurrentInteractionProfile(Session, RightHand, &RightHandState));
+
+	// TODO: Correctly handle mixed interaction profiles?
+	XrPath InteractionProfile = LeftHandState.interactionProfile != XR_NULL_PATH
+		? LeftHandState.interactionProfile
+		: RightHandState.interactionProfile;
+
+	if (InteractionProfile == XR_NULL_PATH)
+	{
+		return TOptional<FName>();
+	}
+
+	uint32 PathLen;
+	char PathBuf[XR_MAX_PATH_LENGTH] = { 0 };
+	if (!XR_ENSURE(xrPathToString(Instance, InteractionProfile, XR_MAX_PATH_LENGTH, &PathLen, PathBuf)))
+	{
+		return FName();
+	}
+
+	const FAnsiStringView PathView(PathBuf, PathLen - 1);
+	if (PathView.StartsWith(ANSITEXTVIEW("/interaction_profiles/oculus/")))
+	{
+		return FName("OculusHMD");
+	}
+	else if (PathView.StartsWith(ANSITEXTVIEW("/interaction_profiles/htc/"))
+		|| PathView.StartsWith(ANSITEXTVIEW("/interaction_profiles/valve/")))
+	{
+		return FName("SteamVR");
+	}
+
+	UE_LOG(LogVirtualScoutingOpenXR, Error, TEXT("Unsupported interaction profile: '%S'"), PathBuf);
+	return FName();
+}
+
+
 void FVirtualScoutingOpenXRExtension::OnVREditingModeEnter()
 {
 	bIsVrEditingModeActive = true;
@@ -243,7 +338,7 @@ XrBool32 XRAPI_CALL FVirtualScoutingOpenXRExtension::XrDebugUtilsMessengerCallba
 	Types[2] = (InMessageTypes & XR_DEBUG_UTILS_MESSAGE_TYPE_PERFORMANCE_BIT_EXT) ? TEXT('P') : TEXT('_');
 	Types[3] = (InMessageTypes & XR_DEBUG_UTILS_MESSAGE_TYPE_CONFORMANCE_BIT_EXT) ? TEXT('C') : TEXT('_');
 
-	FMsg::Logf(__FILE__, __LINE__, LogVPOpenXRDebug.GetCategoryName(), Verbosity,
+	FMsg::Logf(__FILE__, __LINE__, LogVirtualScoutingOpenXRDebug.GetCategoryName(), Verbosity,
 		TEXT("[%s]: %S(): %S"), *Types, InCallbackData->functionName, InCallbackData->message);
 
 	// "A value of XR_TRUE indicates that the application wants to abort this call. [...]
diff --git a/Engine/Plugins/Experimental/VirtualScouting/Source/VirtualScoutingOpenXR/Private/VirtualScoutingOpenXR.h b/Engine/Plugins/Experimental/VirtualScouting/Source/VirtualScoutingOpenXR/Private/VirtualScoutingOpenXR.h
--- a/Engine/Plugins/Experimental/VirtualScouting/Source/VirtualScoutingOpenXR/Private/VirtualScoutingOpenXR.h
+++ b/Engine/Plugins/Experimental/VirtualScouting/Source/VirtualScoutingOpenXR/Private/VirtualScoutingOpenXR.h
@@ -5,11 +5,13 @@
 #if WITH_EDITOR
 
 #include "IOpenXRExtensionPlugin.h"
-
+#include "Async/Future.h"
 #include "OpenXRCore.h"
 
 
-class VIRTUALSCOUTINGOPENXR_API FVirtualScoutingOpenXRExtension : public IOpenXRExtensionPlugin
+class VIRTUALSCOUTINGOPENXR_API FVirtualScoutingOpenXRExtension
+	: public IOpenXRExtensionPlugin
+	, public TSharedFromThis<FVirtualScoutingOpenXRExtension>
 {
 public:
 	FVirtualScoutingOpenXRExtension();
@@ -22,25 +24,37 @@ public:
 	}
 
 	virtual bool GetOptionalExtensions(TArray<const ANSICHAR*>& OutExtensions) override;
+	virtual void OnEvent(XrSession InSession, const XrEventDataBaseHeader* InHeader) override;
 	virtual void PostCreateInstance(XrInstance InInstance) override;
+	virtual void PostCreateSession(XrSession InSession) override;
 	virtual void AddActions(XrInstance InInstance, FCreateActionSetFunc CreateActionSet, FCreateActionFunc CreateAction) override;
 	virtual void GetActiveActionSetsForSync(TArray<XrActiveActionSet>& OutActiveSets) override;
+	virtual void PostSyncActions(XrSession InSession) override;
 	//~ End IOpenXRExtensionPlugin interface
 
+	TFuture<FName>& GetHmdDeviceTypeFuture() { return DeviceTypeFuture; }
+
 private:
 	void OnVREditingModeEnter();
 	void OnVREditingModeExit();
 
+	TOptional<FName> TryGetHmdDeviceType();
+	void TryFulfillDeviceTypePromise();
+
 private:
 	FDelegateHandle InitCompleteDelegate;
 
 	XrDebugUtilsMessengerEXT Messenger = XR_NULL_HANDLE;
 
 	XrInstance Instance = XR_NULL_HANDLE;
+	XrSession Session = XR_NULL_HANDLE;
 	XrActionSet ActionSet = XR_NULL_HANDLE;
 
 	bool bIsVrEditingModeActive = false;
 
+	TPromise<FName> DeviceTypePromise;
+	TFuture<FName> DeviceTypeFuture;
+
 private:
 	static XrBool32 XRAPI_CALL XrDebugUtilsMessengerCallback_Trampoline(
 		XrDebugUtilsMessageSeverityFlagsEXT InMessageSeverity,
diff --git a/Engine/Plugins/Experimental/VirtualScouting/Source/VirtualScoutingOpenXR/Private/VirtualScoutingOpenXRModule.h b/Engine/Plugins/Experimental/VirtualScouting/Source/VirtualScoutingOpenXR/Private/VirtualScoutingOpenXRModule.h
new file mode 100644
--- /dev/null
+++ b/Engine/Plugins/Experimental/VirtualScouting/Source/VirtualScoutingOpenXR/Private/VirtualScoutingOpenXRModule.h
@@ -0,0 +1,34 @@
+// Copyright Epic Games, Inc. All Rights Reserved.
+
+#pragma once
+
+#include "Modules/ModuleInterface.h"
+#include "Modules/ModuleManager.h"
+
+
+class FVirtualScoutingOpenXRExtension;
+
+
+class VIRTUALSCOUTINGOPENXR_API FVirtualScoutingOpenXRModule : public IModuleInterface
+{
+public:
+	/**
+	 * Singleton-like access to this module's interface.  This is just for convenience!
+	 * Beware of calling this during the shutdown phase, though.  Your module might have been unloaded already.
+	 *
+	 * @return Returns singleton instance, loading the module on demand if needed
+	 */
+	static FVirtualScoutingOpenXRModule& Get()
+	{
+		static const FName ModuleName = "VirtualScoutingOpenXR";
+		return FModuleManager::LoadModuleChecked<FVirtualScoutingOpenXRModule>(ModuleName);
+	};
+
+	virtual void StartupModule() override;
+	virtual void ShutdownModule() override;
+
+	const TSharedPtr<FVirtualScoutingOpenXRExtension>& GetOpenXRExt() { return OpenXRExt; }
+
+private:
+	TSharedPtr<FVirtualScoutingOpenXRExtension> OpenXRExt;
+};
diff --git a/Engine/Plugins/Experimental/VirtualScouting/VirtualScouting.uplugin b/Engine/Plugins/Experimental/VirtualScouting/VirtualScouting.uplugin
--- a/Engine/Plugins/Experimental/VirtualScouting/VirtualScouting.uplugin
+++ b/Engine/Plugins/Experimental/VirtualScouting/VirtualScouting.uplugin
@@ -26,6 +26,11 @@
 				"Win64",
 				"Linux"
 			]
+		},
+		{
+			"Name": "VirtualScoutingEditor",
+			"Type": "Editor",
+			"LoadingPhase": "Default"
 		}
 	],
 	"Plugins": [
diff --git a/Engine/Source/Editor/VREditor/Public/VREditorMode.h b/Engine/Source/Editor/VREditor/Public/VREditorMode.h
--- a/Engine/Source/Editor/VREditor/Public/VREditorMode.h
+++ b/Engine/Source/Editor/VREditor/Public/VREditorMode.h
@@ -59,14 +59,25 @@ public:
 	/** Default constructor */
 	UVREditorMode();
 
+	/** Overrides the HMD device type, which is otherwise derived from the XR system name. */
+	void SetHMDDeviceTypeOverride( FName InOverrideType );
+
+	//~ Begin UEditorWorldExtension interface
+
 	/** Initialize the VREditor */
 	virtual void Init() override;
 
 	/** Shutdown the VREditor */
 	virtual void Shutdown() override;
 
+protected:
+	virtual void TransitionWorld(UWorld* NewWorld, EEditorWorldExtensionTransitionState TransitionState) override;
+
+	//~ End UEditorWorldExtension interface
+
+public:
 	/** When the user actually enters the VR Editor mode */
-	void Enter();
+	virtual void Enter();
 
 	/** When the user leaves the VR Editor mode */
 	void Exit( const bool bShouldDisableStereo );
@@ -307,15 +318,15 @@ public:
 	DECLARE_EVENT_OneParam(UVREditorMode, FOnToggleVRModeDebug, bool);
 	FOnToggleVRModeDebug& OnToggleDebugMode() { return OnToggleDebugModeEvent; };
 
+	/** Delegate to be called when async VR mode entry has been completed. */
+	DECLARE_EVENT(UVREditorMode, FOnVRModeEntryComplete);
+	FOnVRModeEntryComplete& OnVRModeEntryComplete() { return OnVRModeEntryCompleteEvent; }
+
 	const TArray<UVREditorInteractor*> GetVRInteractors() const
 	{
 		return Interactors;
 	}
 
-protected:
-
-	virtual void TransitionWorld(UWorld* NewWorld, EEditorWorldExtensionTransitionState TransitionState) override;
-
 private:
 
 	/** Called when the editor is closed */
@@ -340,6 +351,10 @@ protected:
 	// Startup/Shutdown
 	//
 
+	void BeginEntry();
+	void SetupSubsystems();
+	void FinishEntry();
+
 	/** The VR editor window, if it's open right now */
 	TWeakPtr< class SWindow > VREditorWindowWeakPtr;
 
@@ -500,4 +515,10 @@ private:
 
 	/** Event that gets broadcasted when debug mode is toggled. */
 	FOnToggleVRModeDebug OnToggleDebugModeEvent;
+
+	/** Overridden HMD device type. If NAME_None, HMD device type is derived from the XR system name. */
+	FName HMDDeviceTypeOverride = NAME_None;
+
+	/** Event that gets broadcast when async VR mode entry is completed. */
+	FOnVRModeEntryComplete OnVRModeEntryCompleteEvent;
 };
diff --git a/Engine/Source/Editor/VREditor/Public/VRModeSettings.h b/Engine/Source/Editor/VREditor/Public/VRModeSettings.h
--- a/Engine/Source/Editor/VREditor/Public/VRModeSettings.h
+++ b/Engine/Source/Editor/VREditor/Public/VRModeSettings.h
@@ -88,6 +88,10 @@ public:
 	UPROPERTY( EditAnywhere, config, NoClear, Category = "Motion Controllers" )
 	TSoftClassPtr<AVREditorTeleporter> TeleporterClass;
 
+	/** The mode extension to use when UnrealEd is in VR mode. Use VREditorMode to get default editor behavior or select a custom mode. */
+	UPROPERTY(EditAnywhere, config, NoClear, Category = "General")
+	TSoftClassPtr<UVREditorMode> ModeClass;
+
 private:
 #if WITH_EDITOR
 	void PostEditChangeProperty(struct FPropertyChangedEvent& PropertyChangedEvent);
diff --git a/Engine/Source/Editor/VREditor/Teleporter/VREditorTeleporter.cpp b/Engine/Source/Editor/VREditor/Teleporter/VREditorTeleporter.cpp
--- a/Engine/Source/Editor/VREditor/Teleporter/VREditorTeleporter.cpp
+++ b/Engine/Source/Editor/VREditor/Teleporter/VREditorTeleporter.cpp
@@ -243,50 +243,48 @@ void AVREditorTeleporter::OnPreviewInputAction(class FEditorViewportClient& View
 
 	if (VREditorInteractor != nullptr && VREditorInteractor->GetControllerType() == EControllerType::Laser)
 	{
+		const bool bDraggingWorldForTeleport =
+			(Interactor->GetDraggingMode () == EViewportInteractionDraggingMode::World &&
+			(Interactor->GetOtherInteractor () == nullptr || Interactor->GetOtherInteractor ()->GetDraggingMode () != EViewportInteractionDraggingMode::AssistingDrag)) ||
+				(VREditorInteractor->GetControllerType() == EControllerType::Navigation && (Action.ActionType == VRActionTypes::TrackpadDown));
+		if (bDraggingWorldForTeleport &&
+			Action.Event == IE_Pressed &&
+			Action.ActionType == ViewportWorldActionTypes::SelectAndMove)
 		{
-			const bool bDraggingWorldForTeleport =
-				(Interactor->GetDraggingMode () == EViewportInteractionDraggingMode::World &&
-				(Interactor->GetOtherInteractor () == nullptr || Interactor->GetOtherInteractor ()->GetDraggingMode () != EViewportInteractionDraggingMode::AssistingDrag)) ||
-					(VREditorInteractor->GetControllerType() == EControllerType::Navigation && (Action.ActionType == VRActionTypes::TrackpadDown));
-			if (bDraggingWorldForTeleport &&
-				Action.Event == IE_Pressed &&
-				Action.ActionType == ViewportWorldActionTypes::SelectAndMove)
+			if (InteractorTryingTeleport == nullptr && TeleportingState == EState::None)
 			{
-				if (InteractorTryingTeleport == nullptr && TeleportingState == EState::None)
-				{
-					TeleportGoalScale = VRMode->GetWorldScaleFactor ();
-
-					// Checking teleport
-					SetVisibility ( true );
-					InteractorTryingTeleport = Interactor;
+				TeleportGoalScale = VRMode->GetWorldScaleFactor ();
 
-					TeleportingState = EState::Aiming;
-					SetColor ( VRMode->GetColor ( UVREditorMode::EColors::WorldDraggingColor ) );
-					VRMode->GetWorldInteraction ().AllowWorldMovement ( false );
-				}
+				// Checking teleport
+				SetVisibility ( true );
+				InteractorTryingTeleport = Interactor;
 
-				bWasHandled = true;
+				TeleportingState = EState::Aiming;
+				SetColor ( VRMode->GetColor ( UVREditorMode::EColors::WorldDraggingColor ) );
+				VRMode->GetWorldInteraction ().AllowWorldMovement ( false );
 			}
 
-			if (Action.Event == IE_Released &&
-				Action.ActionType == ViewportWorldActionTypes::SelectAndMove &&
-				InteractorTryingTeleport != nullptr && InteractorTryingTeleport == Interactor)
+			bWasHandled = true;
+		}
+
+		if (Action.Event == IE_Released &&
+			Action.ActionType == ViewportWorldActionTypes::SelectAndMove &&
+			InteractorTryingTeleport != nullptr && InteractorTryingTeleport == Interactor)
+		{
+			if (!(VREditorInteractor->IsHoveringOverGizmo () || VREditorInteractor->IsHoveringOverUI ()))
 			{
-				if (!(VREditorInteractor->IsHoveringOverGizmo () || VREditorInteractor->IsHoveringOverUI ()))
-				{
-					// Confirm teleport
-					StartTeleport ( );
-					// Clean everything up
-					bPushedFromEndOfLaser = false;
-					bInitialTeleportAim = true;
-					bShouldBeVisible = false;
-					SetVisibility ( false );
-
-					VREditorInteractor->SetForceShowLaser ( false );
-					InteractorTryingTeleport = nullptr;
-
-					bWasHandled = true;
-				}
+				// Confirm teleport
+				StartTeleport ( );
+				// Clean everything up
+				bPushedFromEndOfLaser = false;
+				bInitialTeleportAim = true;
+				bShouldBeVisible = false;
+				SetVisibility ( false );
+
+				VREditorInteractor->SetForceShowLaser ( false );
+				InteractorTryingTeleport = nullptr;
+
+				bWasHandled = true;
 			}
 		}
 	}
diff --git a/Engine/Source/Editor/VREditor/VREditorInteractor.cpp b/Engine/Source/Editor/VREditor/VREditorInteractor.cpp
--- a/Engine/Source/Editor/VREditor/VREditorInteractor.cpp
+++ b/Engine/Source/Editor/VREditor/VREditorInteractor.cpp
@@ -746,7 +746,7 @@ void UVREditorInteractor::Tick_Implementation( const float DeltaTime )
 
 FName UVREditorInteractor::GetHMDDeviceType() const
 {
-	return (GEngine && GEngine->XRSystem.IsValid()) ? GEngine->XRSystem->GetSystemName() : FName();
+	return GetVRMode().GetHMDDeviceType();
 }
 
 void UVREditorInteractor::CalculateDragRay( float& InOutDragRayLength, float& InOutDragRayVelocity )
@@ -978,59 +978,57 @@ void UVREditorInteractor::HandleInputAxis( FEditorViewportClient& ViewportClient
 {
 	if (!bOutWasHandled && Action.ActionType == TriggerAxis)
 	{
+		const float TriggerPressedThreshold = (GetHMDDeviceType() == OculusDeviceType) ? GetDefault<UVRModeSettings>()->TriggerPressedThreshold_Rift : GetDefault<UVRModeSettings>()->TriggerPressedThreshold_Vive;
+		const float TriggerDeadZone = (GetHMDDeviceType() == OculusDeviceType) ? VREd::TriggerDeadZone_Rift->GetFloat() : VREd::TriggerDeadZone_Vive->GetFloat();
+
+		// Synthesize "lightly pressed" events for the trigger
 		{
-			const float TriggerPressedThreshold = (GetHMDDeviceType() == OculusDeviceType) ? GetDefault<UVRModeSettings>()->TriggerPressedThreshold_Rift : GetDefault<UVRModeSettings>()->TriggerPressedThreshold_Vive;
-			const float TriggerDeadZone = (GetHMDDeviceType() == OculusDeviceType) ? VREd::TriggerDeadZone_Rift->GetFloat() : VREd::TriggerDeadZone_Vive->GetFloat();
+			// Store latest trigger value amount
+			SelectAndMoveTriggerValue = Delta;
 
-			// Synthesize "lightly pressed" events for the trigger
+			if (!bIsTriggerPressed &&			// Don't fire if we are already pressed
+				bHasTriggerBeenReleasedSinceLastPress &&	// Only if we've been fully released since the last time we fired
+				Delta >= TriggerPressedThreshold)
 			{
-				// Store latest trigger value amount
-				SelectAndMoveTriggerValue = Delta;
-
-				if (!bIsTriggerPressed &&			// Don't fire if we are already pressed
-					bHasTriggerBeenReleasedSinceLastPress &&	// Only if we've been fully released since the last time we fired
-					Delta >= TriggerPressedThreshold)
-				{
-					bIsTriggerPressed = true;
-					bHasTriggerBeenReleasedSinceLastPress = false;
-					// Synthesize an input key for this light press
-					const EInputEvent InputEvent = IE_Pressed;
-					const bool bWasLightPressHandled = UViewportInteractor::HandleInputKey( ViewportClient, ControllerMotionSource == FXRMotionControllerBase::LeftHandSourceId ? MotionController_Left_PressedTriggerAxis : MotionController_Right_PressedTriggerAxis, InputEvent );
-				}
-				else if (bIsTriggerPressed && Delta < TriggerPressedThreshold)
-				{
-					bIsTriggerPressed = false;
-
-					// Synthesize an input key for this light press
-					const EInputEvent InputEvent = IE_Released;
-					const bool bWasLightReleaseHandled = UViewportInteractor::HandleInputKey( ViewportClient, ControllerMotionSource == FXRMotionControllerBase::LeftHandSourceId ? MotionController_Left_PressedTriggerAxis : MotionController_Right_PressedTriggerAxis, InputEvent );
-				}
+				bIsTriggerPressed = true;
+				bHasTriggerBeenReleasedSinceLastPress = false;
+				// Synthesize an input key for this light press
+				const EInputEvent InputEvent = IE_Pressed;
+				const bool bWasLightPressHandled = UViewportInteractor::HandleInputKey( ViewportClient, ControllerMotionSource == FXRMotionControllerBase::LeftHandSourceId ? MotionController_Left_PressedTriggerAxis : MotionController_Right_PressedTriggerAxis, InputEvent );
 			}
-
-			if (!bHasTriggerBeenReleasedSinceLastPress && Delta < TriggerDeadZone)
+			else if (bIsTriggerPressed && Delta < TriggerPressedThreshold)
 			{
-				bHasTriggerBeenReleasedSinceLastPress = true;
+				bIsTriggerPressed = false;
+
+				// Synthesize an input key for this light press
+				const EInputEvent InputEvent = IE_Released;
+				const bool bWasLightReleaseHandled = UViewportInteractor::HandleInputKey( ViewportClient, ControllerMotionSource == FXRMotionControllerBase::LeftHandSourceId ? MotionController_Left_PressedTriggerAxis : MotionController_Right_PressedTriggerAxis, InputEvent );
 			}
+		}
 
-			// Synthesize "fully pressed" events for the trigger
-			{
-				const float TriggerFullyPressedThreshold = (GetHMDDeviceType() == OculusDeviceType) ? VREd::TriggerFullyPressedThreshold_Rift->GetFloat() : VREd::TriggerFullyPressedThreshold_Vive->GetFloat();
+		if (!bHasTriggerBeenReleasedSinceLastPress && Delta < TriggerDeadZone)
+		{
+			bHasTriggerBeenReleasedSinceLastPress = true;
+		}
 
-				if (!bIsTriggerFullyPressed &&	// Don't fire if we are already pressed
-					Delta >= TriggerFullyPressedThreshold)
-				{
-					bIsTriggerFullyPressed = true;
+		// Synthesize "fully pressed" events for the trigger
+		{
+			const float TriggerFullyPressedThreshold = (GetHMDDeviceType() == OculusDeviceType) ? VREd::TriggerFullyPressedThreshold_Rift->GetFloat() : VREd::TriggerFullyPressedThreshold_Vive->GetFloat();
 
-					const EInputEvent InputEvent = IE_Pressed;
-					UViewportInteractor::HandleInputKey( ViewportClient, ControllerMotionSource == FXRMotionControllerBase::LeftHandSourceId ? MotionController_Left_FullyPressedTriggerAxis : MotionController_Right_FullyPressedTriggerAxis, InputEvent );
-				}
-				else if (bIsTriggerFullyPressed && Delta < TriggerPressedThreshold)
-				{
-					bIsTriggerFullyPressed = false;
+			if (!bIsTriggerFullyPressed &&	// Don't fire if we are already pressed
+				Delta >= TriggerFullyPressedThreshold)
+			{
+				bIsTriggerFullyPressed = true;
 
-					const EInputEvent InputEvent = IE_Released;
-					UViewportInteractor::HandleInputKey( ViewportClient, ControllerMotionSource == FXRMotionControllerBase::LeftHandSourceId ? MotionController_Left_FullyPressedTriggerAxis : MotionController_Right_FullyPressedTriggerAxis, InputEvent );
-				}
+				const EInputEvent InputEvent = IE_Pressed;
+				UViewportInteractor::HandleInputKey( ViewportClient, ControllerMotionSource == FXRMotionControllerBase::LeftHandSourceId ? MotionController_Left_FullyPressedTriggerAxis : MotionController_Right_FullyPressedTriggerAxis, InputEvent );
+			}
+			else if (bIsTriggerFullyPressed && Delta < TriggerPressedThreshold)
+			{
+				bIsTriggerFullyPressed = false;
+
+				const EInputEvent InputEvent = IE_Released;
+				UViewportInteractor::HandleInputKey( ViewportClient, ControllerMotionSource == FXRMotionControllerBase::LeftHandSourceId ? MotionController_Left_FullyPressedTriggerAxis : MotionController_Right_FullyPressedTriggerAxis, InputEvent );
 			}
 		}
 	}
diff --git a/Engine/Source/Editor/VREditor/VREditorMode.cpp b/Engine/Source/Editor/VREditor/VREditorMode.cpp
--- a/Engine/Source/Editor/VREditor/VREditorMode.cpp
+++ b/Engine/Source/Editor/VREditor/VREditorMode.cpp
@@ -98,7 +98,14 @@ UVREditorMode::UVREditorMode() :
 	WorldInteraction( nullptr ),
 	bFirstTick( true ),
 	AssetContainer( nullptr )
-{ 
+{
+	SetActive(false);
+}
+
+void UVREditorMode::SetHMDDeviceTypeOverride( FName InOverrideType )
+{
+	ensureMsgf(!IsActive(), TEXT("HMD device type override should be specified before VR editor mode is entered."));
+	HMDDeviceTypeOverride = InOverrideType;
 }
 
 void UVREditorMode::Init()
@@ -221,12 +228,19 @@ void UVREditorMode::AllocateInteractors()
 }
 
 void UVREditorMode::Enter()
+{
+	BeginEntry();
+	SetupSubsystems();
+	FinishEntry();
+}
+
+void UVREditorMode::BeginEntry()
 {
 	bWantsToExitMode = false;
 
 	{
-		WorldInteraction->OnPreWorldInteractionTick().AddUObject( this, &UVREditorMode::PreTick );
-		WorldInteraction->OnPostWorldInteractionTick().AddUObject( this, &UVREditorMode::PostTick );
+		WorldInteraction->OnPreWorldInteractionTick().AddUObject(this, &UVREditorMode::PreTick);
+		WorldInteraction->OnPostWorldInteractionTick().AddUObject(this, &UVREditorMode::PostTick);
 	}
 
 
@@ -242,7 +256,7 @@ void UVREditorMode::Enter()
 		TSharedPtr<SLevelViewport> ExistingActiveLevelViewport;
 		{
 			TSharedPtr<IAssetViewport> ActiveLevelViewport = LevelEditor->GetActiveViewportInterface();
-			if(ActiveLevelViewport.IsValid())
+			if (ActiveLevelViewport.IsValid())
 			{
 				ExistingActiveLevelViewport = StaticCastSharedRef< SLevelViewport >(ActiveLevelViewport->AsWidget());
 				ExistingActiveLevelViewport->RemoveAllPreviews(true);
@@ -251,12 +265,12 @@ void UVREditorMode::Enter()
 
 		StartViewport(ExistingActiveLevelViewport);
 
-		if( bActuallyUsingVR )
+		if (bActuallyUsingVR)
 		{
 			// Tell Slate to require a larger pixel distance threshold before the drag starts.  This is important for things
 			// like Content Browser drag and drop.
 			SavedEditorState.DragTriggerDistance = FSlateApplication::Get().GetDragTriggerDistance();
-			FSlateApplication::Get().SetDragTriggerDistance( VREd::SlateDragDistanceOverride->GetFloat() );
+			FSlateApplication::Get().SetDragTriggerDistance(VREd::SlateDragDistanceOverride->GetFloat());
 
 			// When actually in VR, make sure the transform gizmo is big!
 			SavedEditorState.TransformGizmoScale = WorldInteraction->GetTransformGizmoScale();
@@ -285,6 +299,7 @@ void UVREditorMode::Enter()
 			SequencerTab->RequestCloseTab();
 		}
 	}
+
 	// Setup our avatar
 	if (AvatarActor == nullptr)
 	{
@@ -292,74 +307,77 @@ void UVREditorMode::Enter()
 		AvatarActor = SpawnTransientSceneActor<AVREditorAvatarActor>(TEXT("AvatarActor"), bWithSceneComponent);
 		AvatarActor->Init(this);
 
-		WorldInteraction->AddActorToExcludeFromHitTests( AvatarActor );	
+		WorldInteraction->AddActorToExcludeFromHitTests(AvatarActor);
 	}
 
 	// If we're actually using VR, go ahead and disable notifications.  We won't be able to see them in VR
 	// currently, and they can introduce performance issues if they pop up on the desktop
-	if( bActuallyUsingVR )
+	if (bActuallyUsingVR)
 	{
-		FSlateNotificationManager::Get().SetAllowNotifications( false );
+		FSlateNotificationManager::Get().SetAllowNotifications(false);
 	}
 
-	// Setup sub systems
-	{
-		// Setup world interaction
-		// We need input preprocessing for VR so that we can receive motion controller input without any viewports having 
-		// to be focused.  This is mainly because Slate UI injected into the 3D world can cause focus to be lost unexpectedly,
-		// but we need the user to still be able to interact with UI.
-		WorldInteraction->SetUseInputPreprocessor( true );
-
-		// Motion controllers
-		AllocateInteractors();
-
-		if( bActuallyUsingVR )
-		{
-			// When actually using VR devices, we don't want a mouse cursor interactor
-			WorldInteraction->ReleaseMouseCursorInteractor();
-		}
+	/** This will make sure this is not ticking after the editor has been closed. */
+	GEditor->OnEditorClose().AddUObject(this, &UVREditorMode::OnEditorClosed);
+}
 
-		// Setup the UI system
-		UISystem = NewObject<UVREditorUISystem>();
-		UISystem->Init(this);
+void UVREditorMode::SetupSubsystems()
+{
+	// Setup world interaction
+	// We need input preprocessing for VR so that we can receive motion controller input without any viewports having 
+	// to be focused.  This is mainly because Slate UI injected into the 3D world can cause focus to be lost unexpectedly,
+	// but we need the user to still be able to interact with UI.
+	WorldInteraction->SetUseInputPreprocessor( true );
 
-		PlacementSystem = NewObject<UVREditorPlacement>();
-		PlacementSystem->Init(this);
+	// Motion controllers
+	AllocateInteractors();
 
-		// Setup teleporter
-		const TSoftClassPtr<AVREditorTeleporter> TeleporterClassSoft = GetDefault<UVRModeSettings>()->TeleporterClass;
-		TeleporterClassSoft.LoadSynchronous();
+	if( bActuallyUsingVR )
+	{
+		// When actually using VR devices, we don't want a mouse cursor interactor
+		WorldInteraction->ReleaseMouseCursorInteractor();
+	}
 
-		if (TeleporterClassSoft.IsValid())
-		{
-			TeleportActor = CastChecked<AVREditorTeleporter>( SpawnTransientSceneActor(TeleporterClassSoft.Get(), TEXT( "Teleporter" ), true ) );
-		}
+	// Setup the UI system
+	UISystem = NewObject<UVREditorUISystem>();
+	UISystem->Init(this);
 
-		if( !TeleportActor )
-		{
-			TeleportActor = SpawnTransientSceneActor<AVREditorTeleporter>( TEXT( "Teleporter" ), true );
-		}
+	PlacementSystem = NewObject<UVREditorPlacement>();
+	PlacementSystem->Init(this);
 
-		check( TeleportActor );
-		TeleportActor->Init( this );
-		WorldInteraction->AddActorToExcludeFromHitTests( TeleportActor );
+	// Setup teleporter
+	const TSoftClassPtr<AVREditorTeleporter> TeleporterClassSoft = GetDefault<UVRModeSettings>()->TeleporterClass;
+	TeleporterClassSoft.LoadSynchronous();
 
-		// Setup autoscaler
-		AutoScalerSystem = NewObject<UVREditorAutoScaler>();
-		AutoScalerSystem->Init( this );
+	if (TeleporterClassSoft.IsValid())
+	{
+		TeleportActor = CastChecked<AVREditorTeleporter>( SpawnTransientSceneActor(TeleporterClassSoft.Get(), TEXT( "Teleporter" ), true ) );
+	}
 
-		for (UVREditorInteractor* Interactor : Interactors)
-		{
-			Interactor->SetupComponent( AvatarActor );
-		}
+	if( !TeleportActor )
+	{
+		TeleportActor = SpawnTransientSceneActor<AVREditorTeleporter>( TEXT( "Teleporter" ), true );
 	}
 
+	check( TeleportActor );
+	TeleportActor->Init( this );
+	WorldInteraction->AddActorToExcludeFromHitTests( TeleportActor );
 
-	/** This will make sure this is not ticking after the editor has been closed. */
-	GEditor->OnEditorClose().AddUObject( this, &UVREditorMode::OnEditorClosed );
+	// Setup autoscaler
+	AutoScalerSystem = NewObject<UVREditorAutoScaler>();
+	AutoScalerSystem->Init( this );
 
+	for (UVREditorInteractor* Interactor : Interactors)
+	{
+		Interactor->SetupComponent( AvatarActor );
+	}
+}
+
+void UVREditorMode::FinishEntry()
+{
 	bFirstTick = true;
 	SetActive(true);
+	OnVRModeEntryCompleteEvent.Broadcast();
 }
 
 void UVREditorMode::Exit(const bool bShouldDisableStereo)
@@ -723,6 +741,11 @@ void UVREditorMode::CycleTransformGizmoHandleType()
 
 FName UVREditorMode::GetHMDDeviceType() const
 {
+	if (HMDDeviceTypeOverride != NAME_None)
+	{
+		return HMDDeviceTypeOverride;
+	}
+
 	return GEngine->XRSystem.IsValid() ? GEngine->XRSystem->GetSystemName() : FName();
 }
 
diff --git a/Engine/Source/Editor/VREditor/VREditorModeManager.cpp b/Engine/Source/Editor/VREditor/VREditorModeManager.cpp
--- a/Engine/Source/Editor/VREditor/VREditorModeManager.cpp
+++ b/Engine/Source/Editor/VREditor/VREditorModeManager.cpp
@@ -36,7 +36,11 @@ FVREditorModeManager::FVREditorModeManager() :
 
 FVREditorModeManager::~FVREditorModeManager()
 {
-	CurrentVREditorMode = nullptr;
+	if (CurrentVREditorMode)
+	{
+		CurrentVREditorMode->OnVRModeEntryComplete().RemoveAll(this);
+		CurrentVREditorMode = nullptr;
+	}
 }
 
 void FVREditorModeManager::Tick( const float DeltaTime )
@@ -166,6 +170,15 @@ void FVREditorModeManager::AddReferencedObjects( FReferenceCollector& Collector
 	Collector.AddReferencedObject( CurrentVREditorMode );
 }
 
+void FVREditorModeManager::HandleModeEntryComplete()
+{
+	// Connects the mode UObject's event to the module delegate.
+	if (CurrentVREditorMode->IsActuallyUsingVR())
+	{
+		OnVREditingModeEnterHandle.Broadcast();
+	}
+}
+
 void FVREditorModeManager::StartVREditorMode( const bool bForceWithoutHMD )
 {
 	if (!IsEngineExitRequested())
@@ -192,8 +205,13 @@ void FVREditorModeManager::StartVREditorMode( const bool bForceWithoutHMD )
 			}
 
 			// Create vr editor mode.
-			VRMode = NewObject<UVREditorMode>();
+			const TSoftClassPtr<UVREditorMode> ModeClassSoft = GetDefault<UVRModeSettings>()->ModeClass;
+			ModeClassSoft.LoadSynchronous();
+			check(ModeClassSoft.IsValid());
+
+			VRMode = NewObject<UVREditorMode>(GetTransientPackage(), ModeClassSoft.Get());
 			check(VRMode != nullptr);
+			VRMode->OnVRModeEntryComplete().AddRaw(this, &FVREditorModeManager::HandleModeEntryComplete);
 			ExtensionCollection->AddExtension(VRMode);
 		}
 
@@ -207,11 +225,6 @@ void FVREditorModeManager::StartVREditorMode( const bool bForceWithoutHMD )
 		CurrentVREditorMode->SetActuallyUsingVR( !bForceWithoutHMD );
 
 		CurrentVREditorMode->Enter();
-
-		if (CurrentVREditorMode->IsActuallyUsingVR())
-		{
-			OnVREditingModeEnterHandle.Broadcast();
-		}
 	}
 }
 
diff --git a/Engine/Source/Editor/VREditor/VREditorModeManager.h b/Engine/Source/Editor/VREditor/VREditorModeManager.h
--- a/Engine/Source/Editor/VREditor/VREditorModeManager.h
+++ b/Engine/Source/Editor/VREditor/VREditorModeManager.h
@@ -58,7 +58,7 @@ public:
 	// End FGCObject
 
 	/** Return a multicast delegate which is executed when VR mode starts. */
-	FOnVREditingModeEnter& OnVREditingModeEnter() { return OnVREditingModeEnterHandle;  }
+	FOnVREditingModeEnter& OnVREditingModeEnter() { return OnVREditingModeEnterHandle; }
 
 	/** Return a multicast delegate which is executed when VR mode stops. */
 	FOnVREditingModeExit& OnVREditingModeExit() { return OnVREditingModeExitHandle; }
@@ -71,6 +71,8 @@ private:
 	/** Broadcasts when VR mode is stopped */
 	FOnVREditingModeExit OnVREditingModeExitHandle;
 
+	void HandleModeEntryComplete();
+
 	/** Saves the WorldToMeters and enters the mode belonging to GWorld */
 	void StartVREditorMode( const bool bForceWithoutHMD );
 
