Virtual Scouting: Non-unity build fix.

diff --git a/Engine/Plugins/Experimental/VirtualScouting/Source/VirtualScoutingEditor/Private/VirtualScoutingMode.cpp b/Engine/Plugins/Experimental/VirtualScouting/Source/VirtualScoutingEditor/Private/VirtualScoutingMode.cpp
--- a/Engine/Plugins/Experimental/VirtualScouting/Source/VirtualScoutingEditor/Private/VirtualScoutingMode.cpp
+++ b/Engine/Plugins/Experimental/VirtualScouting/Source/VirtualScoutingEditor/Private/VirtualScoutingMode.cpp
@@ -4,6 +4,10 @@
 #include "VirtualScoutingOpenXR.h"
 #include "VirtualScoutingOpenXRModule.h"
 
+#if WITH_EDITOR
+#include "Editor.h"
+#endif
+
 #include "IXRTrackingSystem.h"
 #include "Logging/LogMacros.h"
 #include "Modules/ModuleManager.h"
@@ -47,8 +51,10 @@ void UVirtualScoutingMode::Enter()
 		return;
 	}
 
+#if WITH_EDITOR
 	// This causes FOpenXRInput to rebuild and reattach actions.
 	FEditorDelegates::OnActionAxisMappingsChanged.Broadcast();
+#endif
 
 	// Split the mode entry into two phases. This is necessary because we have to poll OpenXR for
 	// the active interaction profile and translate it into a legacy plugin name, but OpenXR may
