Virtual Scouting: Log an error and refuse to launch if r.PostProcessing.PropagateAlpha is not set to 0.

diff --git a/Engine/Plugins/Experimental/VirtualScouting/Source/VirtualScoutingEditor/Private/VirtualScoutingMode.cpp b/Engine/Plugins/Experimental/VirtualScouting/Source/VirtualScoutingEditor/Private/VirtualScoutingMode.cpp
--- a/Engine/Plugins/Experimental/VirtualScouting/Source/VirtualScoutingEditor/Private/VirtualScoutingMode.cpp
+++ b/Engine/Plugins/Experimental/VirtualScouting/Source/VirtualScoutingEditor/Private/VirtualScoutingMode.cpp
@@ -9,6 +9,7 @@
 #endif
 
 #include "IOpenXRHMDModule.h"
+#include "IVREditorModule.h"
 #include "IXRTrackingSystem.h"
 #include "Logging/LogMacros.h"
 #include "Modules/ModuleManager.h"
@@ -70,6 +71,13 @@ void UVirtualScoutingMode::Enter()
 	if (!XrExt)
 	{
 		UE_LOG(LogVirtualScouting, Error, TEXT("OpenXR extension plugin invalid"));
+		IVREditorModule::Get().EnableVREditor(false);
+		return;
+	}
+
+	if (!ValidateSettings())
+	{
+		IVREditorModule::Get().EnableVREditor(false);
 		return;
 	}
 
@@ -103,4 +111,22 @@ void UVirtualScoutingMode::Enter()
 }
 
 
+bool UVirtualScoutingMode::ValidateSettings()
+{
+	bool bSettingsValid = true;
+
+	IConsoleManager& ConsoleMgr = IConsoleManager::Get();
+	if (IConsoleVariable* PropagateAlpha = ConsoleMgr.FindConsoleVariable(TEXT("r.PostProcessing.PropagateAlpha")))
+	{
+		if (PropagateAlpha->GetInt() != 0)
+		{
+			UE_LOG(LogVirtualScouting, Error, TEXT("r.PostProcessing.PropagateAlpha must be set to 0 (and requires an engine restart)"));
+			bSettingsValid = false;
+		}
+	}
+
+	return bSettingsValid;
+}
+
+
 #undef LOCTEXT_NAMESPACE
diff --git a/Engine/Plugins/Experimental/VirtualScouting/Source/VirtualScoutingEditor/Private/VirtualScoutingMode.h b/Engine/Plugins/Experimental/VirtualScouting/Source/VirtualScoutingEditor/Private/VirtualScoutingMode.h
--- a/Engine/Plugins/Experimental/VirtualScouting/Source/VirtualScoutingEditor/Private/VirtualScoutingMode.h
+++ b/Engine/Plugins/Experimental/VirtualScouting/Source/VirtualScoutingEditor/Private/VirtualScoutingMode.h
@@ -21,4 +21,7 @@ public:
 
 	virtual void Enter() override;
 	//~ End UVREditorMode interface
+
+protected:
+	bool ValidateSettings();
 };
