Virtual Scouting: Rename current mode class from UVirtualScoutingMode to UVPScoutingMode, to reflect its dependence on the Virtual Production Utilities plugin.

#rnx

diff --git a/Engine/Plugins/Experimental/VirtualScouting/Source/VirtualScoutingEditor/Private/VPScoutingMode.cpp b/Engine/Plugins/Experimental/VirtualScouting/Source/VirtualScoutingEditor/Private/VPScoutingMode.cpp
new file mode 100644
--- /dev/null
+++ b/Engine/Plugins/Experimental/VirtualScouting/Source/VirtualScoutingEditor/Private/VPScoutingMode.cpp
@@ -0,0 +1,132 @@
+// Copyright Epic Games, Inc. All Rights Reserved.
+
+#include "VPScoutingMode.h"
+#include "VirtualScoutingOpenXR.h"
+#include "VirtualScoutingOpenXRModule.h"
+
+#if WITH_EDITOR
+#include "Editor.h"
+#endif
+
+#include "IOpenXRHMDModule.h"
+#include "IVREditorModule.h"
+#include "IXRTrackingSystem.h"
+#include "Logging/LogMacros.h"
+#include "Modules/ModuleManager.h"
+
+
+#define LOCTEXT_NAMESPACE "VirtualScouting"
+
+
+DECLARE_LOG_CATEGORY_EXTERN(LogVirtualScouting, Log, All);
+DEFINE_LOG_CATEGORY(LogVirtualScouting);
+
+
+static const FName OpenXRSystemName = FName(TEXT("OpenXR"));
+
+
+class FVirtualScoutingEditorModule : public IModuleInterface
+{
+private:
+	virtual void StartupModule() override
+	{
+	}
+
+	virtual void ShutdownModule() override
+	{
+	}
+};
+
+
+IMPLEMENT_MODULE(FVirtualScoutingEditorModule, VirtualScoutingEditor);
+
+
+UVPScoutingMode::UVPScoutingMode(const FObjectInitializer& ObjectInitializer)
+{
+	InteractorClass = FSoftObjectPath("/VirtualProductionUtilities/VR/VirtualScoutingInteractor.VirtualScoutingInteractor_C");
+	TeleporterClass = FSoftObjectPath("/VirtualProductionUtilities/VR/VirtualScoutingTeleporter.VirtualScoutingTeleporter_C");
+}
+
+
+bool UVPScoutingMode::NeedsSyntheticDpad()
+{
+	if (!GEngine->XRSystem.IsValid() || GEngine->XRSystem->GetSystemName() != OpenXRSystemName)
+	{
+		return Super::NeedsSyntheticDpad();
+	}
+
+	return !IOpenXRHMDModule::Get().IsExtensionEnabled("XR_EXT_dpad_binding");
+}
+
+
+void UVPScoutingMode::Enter()
+{
+	if (!GEngine->XRSystem.IsValid() || GEngine->XRSystem->GetSystemName() != OpenXRSystemName)
+	{
+		Super::Enter();
+		return;
+	}
+
+	TSharedPtr<FVirtualScoutingOpenXRExtension> XrExt = FVirtualScoutingOpenXRModule::Get().GetOpenXRExt();
+	if (!XrExt)
+	{
+		UE_LOG(LogVirtualScouting, Error, TEXT("OpenXR extension plugin invalid"));
+		IVREditorModule::Get().EnableVREditor(false);
+		return;
+	}
+
+	if (!ValidateSettings())
+	{
+		IVREditorModule::Get().EnableVREditor(false);
+		return;
+	}
+
+#if WITH_EDITOR
+	// This causes FOpenXRInput to rebuild and reattach actions.
+	FEditorDelegates::OnActionAxisMappingsChanged.Broadcast();
+#endif
+
+	// Split the mode entry into two phases. This is necessary because we have to poll OpenXR for
+	// the active interaction profile and translate it into a legacy plugin name, but OpenXR may
+	// not return the correct interaction profile for several frames after the OpenXR session
+	// (stereo rendering) has started, and we need to defer creation of the interactors, etc.
+	BeginEntry();
+
+	XrExt->GetHmdDeviceTypeFuture() = XrExt->GetHmdDeviceTypeFuture().Next([this](FName DeviceType)
+	{
+		if (DeviceType != NAME_None)
+		{
+			SetHMDDeviceTypeOverride(DeviceType);
+		}
+		else
+		{
+			UE_LOG(LogVirtualScouting, Error, TEXT("Unable to map legacy HMD device type"));
+		}
+
+		SetupSubsystems();
+		FinishEntry();
+
+		return DeviceType;
+	});
+}
+
+
+bool UVPScoutingMode::ValidateSettings()
+{
+	bool bSettingsValid = true;
+
+	IConsoleManager& ConsoleMgr = IConsoleManager::Get();
+	if (IConsoleVariable* PropagateAlpha = ConsoleMgr.FindConsoleVariable(TEXT("r.PostProcessing.PropagateAlpha")))
+	{
+		if (PropagateAlpha->GetInt() != 0)
+		{
+			UE_LOG(LogVirtualScouting, Error, TEXT("r.PostProcessing.PropagateAlpha must be set to 0 (and requires an engine restart)"));
+			bSettingsValid = false;
+		}
+	}
+
+	return bSettingsValid;
+}
+
+
+#undef LOCTEXT_NAMESPACE
diff --git a/Engine/Plugins/Experimental/VirtualScouting/Source/VirtualScoutingEditor/Private/VPScoutingMode.h b/Engine/Plugins/Experimental/VirtualScouting/Source/VirtualScoutingEditor/Private/VPScoutingMode.h
new file mode 100644
--- /dev/null
+++ b/Engine/Plugins/Experimental/VirtualScouting/Source/VirtualScoutingEditor/Private/VPScoutingMode.h
@@ -0,0 +1,27 @@
+// Copyright Epic Games, Inc. All Rights Reserved.
+
+#pragma once
+
+#include "VREditorMode.h"
+#include "VPScoutingMode.generated.h"
+
+
+UCLASS(meta=(DisplayName="Virtual Scouting Mode"))
+class UVPScoutingMode : public UVREditorMode
+{
+	GENERATED_BODY()
+
+public:
+	UVPScoutingMode(const FObjectInitializer& ObjectInitializer);
+
+	//~ Begin UVREditorMode interface
+	virtual bool NeedsSyntheticDpad() override;
+
+	virtual bool ShouldDisplayExperimentalWarningOnEntry() const override { return false; }
+
+	virtual void Enter() override;
+	//~ End UVREditorMode interface
+
+protected:
+	bool ValidateSettings();
+};
diff --git a/Engine/Plugins/Experimental/VirtualScouting/Source/VirtualScoutingEditor/Private/VirtualScoutingMode.cpp b/Engine/Plugins/Experimental/VirtualScouting/Source/VirtualScoutingEditor/Private/VirtualScoutingMode.cpp
deleted file mode 100644
diff --git a/Engine/Plugins/Experimental/VirtualScouting/Source/VirtualScoutingEditor/Private/VirtualScoutingMode.h b/Engine/Plugins/Experimental/VirtualScouting/Source/VirtualScoutingEditor/Private/VirtualScoutingMode.h
deleted file mode 100644
