VirtualScoutingEditor should not allow Mac because its build dependency OpenXRHMD does not support Mac

diff --git a/Engine/Plugins/Experimental/VirtualScouting/VirtualScouting.uplugin b/Engine/Plugins/Experimental/VirtualScouting/VirtualScouting.uplugin
--- a/Engine/Plugins/Experimental/VirtualScouting/VirtualScouting.uplugin
+++ b/Engine/Plugins/Experimental/VirtualScouting/VirtualScouting.uplugin
@@ -30,7 +30,11 @@
 		{
 			"Name": "VirtualScoutingEditor",
 			"Type": "Editor",
-			"LoadingPhase": "Default"
+			"LoadingPhase": "Default",
+			"PlatformAllowList": [
+				"Win64",
+				"Linux"
+			]
 		}
 	],
 	"Plugins": [
