Virtual Scouting: When Virtual Scouting exits early due to PropagateAlpha cvar, add a toast/notification to make it more obvious to the user what's happening, and fix a regression that causes a crash.

Also includes speculative PVS-Studio static analysis fix + explanatory comment for VREditorModeManager; followup to CL 23741019.

diff --git a/Engine/Plugins/Experimental/VirtualScouting/Source/VirtualScoutingEditor/Private/VPScoutingMode.cpp b/Engine/Plugins/Experimental/VirtualScouting/Source/VirtualScoutingEditor/Private/VPScoutingMode.cpp
--- a/Engine/Plugins/Experimental/VirtualScouting/Source/VirtualScoutingEditor/Private/VPScoutingMode.cpp
+++ b/Engine/Plugins/Experimental/VirtualScouting/Source/VirtualScoutingEditor/Private/VPScoutingMode.cpp
@@ -8,9 +8,11 @@
 #include "Editor.h"
 #endif
 
+#include "Framework/Notifications/NotificationManager.h"
 #include "IOpenXRHMDModule.h"
 #include "IVREditorModule.h"
 #include "IXRTrackingSystem.h"
+#include "Widgets/Notifications/SNotificationList.h"
 
 
 #define LOCTEXT_NAMESPACE "VirtualScouting"
@@ -125,7 +127,7 @@ bool UVPScoutingMode::ValidateSettings()
 	{
 		if (PropagateAlpha->GetInt() != 0)
 		{
-			UE_LOG(LogVirtualScouting, Error, TEXT("r.PostProcessing.PropagateAlpha must be set to 0 (and requires an engine restart)"));
+			InvalidSettingNotification(LOCTEXT("InvalidCvarPropagateAlpha", "r.PostProcessing.PropagateAlpha must be set to 0 (and requires an engine restart)"));
 			bSettingsValid = false;
 		}
 	}
@@ -134,4 +136,20 @@ bool UVPScoutingMode::ValidateSettings()
 }
 
 
+void UVPScoutingMode::InvalidSettingNotification(const FText& ErrorDetails)
+{
+	FText NotificationHeading = LOCTEXT("InvalidSettingNotificationTitle", "Unable to initialize Virtual Scouting");
+
+	UE_LOG(LogVirtualScouting, Error, TEXT("%s: %s"), *NotificationHeading.ToString(), *ErrorDetails.ToString());
+
+	FNotificationInfo Info(NotificationHeading);
+	Info.SubText = ErrorDetails;
+	TSharedPtr<SNotificationItem> Notification = FSlateNotificationManager::Get().AddNotification(Info);
+	if (Notification)
+	{
+		Notification->SetCompletionState(SNotificationItem::CS_Fail);
+	}
+}
+
+
 #undef LOCTEXT_NAMESPACE
diff --git a/Engine/Plugins/Experimental/VirtualScouting/Source/VirtualScoutingEditor/Private/VPScoutingMode.h b/Engine/Plugins/Experimental/VirtualScouting/Source/VirtualScoutingEditor/Private/VPScoutingMode.h
--- a/Engine/Plugins/Experimental/VirtualScouting/Source/VirtualScoutingEditor/Private/VPScoutingMode.h
+++ b/Engine/Plugins/Experimental/VirtualScouting/Source/VirtualScoutingEditor/Private/VPScoutingMode.h
@@ -24,4 +24,6 @@ public:
 
 protected:
 	bool ValidateSettings();
+
+	void InvalidSettingNotification(const FText& ErrorDetails);
 };
diff --git a/Engine/Plugins/Experimental/VirtualScouting/Source/VirtualScoutingEditor/VirtualScoutingEditor.Build.cs b/Engine/Plugins/Experimental/VirtualScouting/Source/VirtualScoutingEditor/VirtualScoutingEditor.Build.cs
--- a/Engine/Plugins/Experimental/VirtualScouting/Source/VirtualScoutingEditor/VirtualScoutingEditor.Build.cs
+++ b/Engine/Plugins/Experimental/VirtualScouting/Source/VirtualScoutingEditor/VirtualScoutingEditor.Build.cs
@@ -26,6 +26,7 @@ public class VirtualScoutingEditor : ModuleRules
 		PrivateDependencyModuleNames.AddRange(new string[]
 			{
 				"OpenXRHMD",
+				"Slate",
 				"UnrealEd",
 				"VREditor",
 				"VirtualScoutingOpenXR",
diff --git a/Engine/Source/Editor/VREditor/VREditorMode.cpp b/Engine/Source/Editor/VREditor/VREditorMode.cpp
--- a/Engine/Source/Editor/VREditor/VREditorMode.cpp
+++ b/Engine/Source/Editor/VREditor/VREditorMode.cpp
@@ -387,7 +387,11 @@ void UVREditorMode::Exit(const bool bShouldDisableStereo)
 	const FSavedEditorState& SavedEditorState = static_cast<const FSavedEditorState&>(SavedEditorStateChecked());
 
 	{
-		GetLevelViewportPossessedForVR().RemoveAllPreviews(false);
+		if (TSharedPtr<SLevelViewport> VrViewport = GetVrLevelViewport())
+		{
+			VrViewport->RemoveAllPreviews(false);
+		}
+
 		GEditor->SelectNone(true, true, false);
 		GEditor->NoteSelectionChange();
 		FVREditorActionCallbacks::ChangeEditorModes(FBuiltinEditorModes::EM_Default);
diff --git a/Engine/Source/Editor/VREditor/VREditorModeManager.cpp b/Engine/Source/Editor/VREditor/VREditorModeManager.cpp
--- a/Engine/Source/Editor/VREditor/VREditorModeManager.cpp
+++ b/Engine/Source/Editor/VREditor/VREditorModeManager.cpp
@@ -197,6 +197,7 @@ void FVREditorModeManager::StartVREditorMode( const bool bForceWithoutHMD )
 			check(VRMode != nullptr);
 			VRMode->OnVRModeEntryComplete().AddRaw(this, &FVREditorModeManager::HandleModeEntryComplete);
 			ExtensionCollection->AddExtension(VRMode);
+			CurrentVREditorMode = VRMode;
 		}
 
 		// Tell the level editor we want to be notified when selection changes
@@ -204,10 +205,11 @@ void FVREditorModeManager::StartVREditorMode( const bool bForceWithoutHMD )
 			FLevelEditorModule& LevelEditor = FModuleManager::LoadModuleChecked<FLevelEditorModule>( "LevelEditor" );
 			LevelEditor.OnMapChanged().AddRaw( this, &FVREditorModeManager::OnMapChanged );
 		}
-	
-		CurrentVREditorMode = VRMode;
-		CurrentVREditorMode->SetActuallyUsingVR( !bForceWithoutHMD );
-		CurrentVREditorMode->Enter();
+
+		VRMode->SetActuallyUsingVR( !bForceWithoutHMD );
+		VRMode->Enter();
+		// Note: Enter() may, in the case of early init failure, re-enter CloseVREditor,
+		// causing CurrentVREditorMode to be nullified/invalidated beyond this point
 
 		if (FEngineAnalytics::IsAvailable())
 		{
