Fix additional C4702 warnings
#rnx

diff --git a/Engine/Plugins/Experimental/VirtualScouting/Source/VirtualScoutingEditor/Private/VirtualScoutingBlueprintFunctionLibrary.cpp b/Engine/Plugins/Experimental/VirtualScouting/Source/VirtualScoutingEditor/Private/VirtualScoutingBlueprintFunctionLibrary.cpp
--- a/Engine/Plugins/Experimental/VirtualScouting/Source/VirtualScoutingEditor/Private/VirtualScoutingBlueprintFunctionLibrary.cpp
+++ b/Engine/Plugins/Experimental/VirtualScouting/Source/VirtualScoutingEditor/Private/VirtualScoutingBlueprintFunctionLibrary.cpp
@@ -11,12 +11,11 @@
 
 bool UVirtualScoutingBlueprintFunctionLibrary::CheckIsWithEditor()
 {
-
 #if WITH_EDITOR
 	return true;
-#endif
-	
+#else
 	return false;
+#endif
 }
 
 bool UVirtualScoutingBlueprintFunctionLibrary::DeleteActors(const TArray<AActor*>& InActorsToDelete)
diff --git a/Engine/Source/Runtime/SlateCore/Public/Widgets/SWidget.h b/Engine/Source/Runtime/SlateCore/Public/Widgets/SWidget.h
--- a/Engine/Source/Runtime/SlateCore/Public/Widgets/SWidget.h
+++ b/Engine/Source/Runtime/SlateCore/Public/Widgets/SWidget.h
@@ -817,14 +817,16 @@ public:
 			}
 			return StatID;
 		}
+		return TStatId(); // not doing stats at the moment, or ever
 #elif ENABLE_STATNAMEDEVENTS
 		if (!StatID.IsValidStat() && GCycleStatsShouldEmitNamedEvents)
 		{
 			CreateStatID();
 		}
 		return StatID;
-#endif
+#else
 		return TStatId(); // not doing stats at the moment, or ever
+#endif
 	}
 
 	UE_DEPRECATED(4.24, "GetRelativeLayoutScale(int32 ChildIndex, float LayoutScaleMultiplier), your widget will also need to set bHasRelativeLayoutScale in their Construct/ctor.")
