Virtual Scouting - rename Camera aperture array to Viewfinder to avoid ambiguity with the VPSettings plugin

diff --git a/Engine/Plugins/Experimental/VirtualScouting/Source/VirtualScouting/Public/VirtualScoutingSettings.h b/Engine/Plugins/Experimental/VirtualScouting/Source/VirtualScouting/Public/VirtualScoutingSettings.h
--- a/Engine/Plugins/Experimental/VirtualScouting/Source/VirtualScouting/Public/VirtualScoutingSettings.h
+++ b/Engine/Plugins/Experimental/VirtualScouting/Source/VirtualScouting/Public/VirtualScoutingSettings.h
@@ -19,17 +19,17 @@ public:
 	UPROPERTY(Config, EditAnywhere, BlueprintReadWrite, Category="Virtual Scouting", meta=(DisplayName="Show Measurements in Imperial Units"))
 	bool bUseImperial = false;
 
-	UPROPERTY(Config, EditAnywhere, BlueprintReadWrite, Category="Virtual Scouting", meta=(DisplayName="Camera Apertures"))
-	TArray<float> CameraApertureArray = {1.0,1.4,2.0,2.8,4.0,5.6,8.0,11.0,16.0,22.0};
+	UPROPERTY(Config, EditAnywhere, BlueprintReadWrite, Category="Virtual Scouting", meta=(DisplayName="Viewfinder Apertures"))
+	TArray<float> ViewfinderApertureArray = {1.0,1.4,2.0,2.8,4.0,5.6,8.0,11.0,16.0,22.0};
 
-	UFUNCTION(BlueprintPure, Category="Virtual Scouting")
+	UFUNCTION(BlueprintPure, Category="Virtual Scouting", DisplayName="Virtual Scouting Settings")
 	static UVirtualScoutingSettings* GetVirtualScoutingSettings();
 };
 
 /**
  * Per user settings for Virtual Scouting Editor.
  */
-UCLASS(Config=EditorPerProjectUserSettings, DisplayName="Virtual Scouting Editor")
+UCLASS(Config=EditorPerProjectUserSettings, DisplayName="Virtual Scouting Editor Settings")
 class VIRTUALSCOUTING_API UVirtualScoutingEditorSettings : public UObject
 {
 	GENERATED_BODY()
