Virtual Scouting - added  array of masking ratios to scouting settings class.

diff --git a/Engine/Plugins/Experimental/VirtualScouting/Source/VirtualScouting/Public/VirtualScoutingSettings.h b/Engine/Plugins/Experimental/VirtualScouting/Source/VirtualScouting/Public/VirtualScoutingSettings.h
--- a/Engine/Plugins/Experimental/VirtualScouting/Source/VirtualScouting/Public/VirtualScoutingSettings.h
+++ b/Engine/Plugins/Experimental/VirtualScouting/Source/VirtualScouting/Public/VirtualScoutingSettings.h
@@ -20,7 +20,10 @@ public:
 	bool bUseImperial = false;
 
 	UPROPERTY(Config, EditAnywhere, BlueprintReadWrite, Category="Virtual Scouting", meta=(DisplayName="Viewfinder Apertures"))
-	TArray<float> ViewfinderApertureArray = {1.0,1.4,2.0,2.8,4.0,5.6,8.0,11.0,16.0,22.0};
+	TArray<float> ViewfinderApertureArray = {1.2, 2.0, 2.8, 4.0, 5.6, 8.0, 11.0, 16.0, 22.0};
+
+	UPROPERTY(Config, EditAnywhere, BlueprintReadWrite, Category="Virtual Scouting", meta=(DisplayName="Viewfinder Monitor Masks"))
+	TArray<float> ViewfinderMaskArray = {1.33, 1.66, 1.78, 2.0, 2.35, 2.39 };
 
 	UFUNCTION(BlueprintPure, Category="Virtual Scouting", DisplayName="Virtual Scouting Settings")
 	static UVirtualScoutingSettings* GetVirtualScoutingSettings();
