virtual Scouting - added missing CineCaptureComponent plugin dependency

diff --git a/Engine/Plugins/Experimental/VirtualScouting/VirtualScouting.uplugin b/Engine/Plugins/Experimental/VirtualScouting/VirtualScouting.uplugin
--- a/Engine/Plugins/Experimental/VirtualScouting/VirtualScouting.uplugin
+++ b/Engine/Plugins/Experimental/VirtualScouting/VirtualScouting.uplugin
@@ -100,6 +100,10 @@
 		{
 			"Name": "MultiUserClient",
 			"Enabled": true
+		},
+		{
+			"Name": "CineCameraSceneCapture",
+			"Enabled": true
 		}
 	]
 }
