Virtual Scouting - added plugin dependency to LiveLink. Reparented annotation actor to remove dependency to VPUtils

#virtualized

diff --git a/Engine/Plugins/Experimental/VirtualScouting/VirtualScouting.uplugin b/Engine/Plugins/Experimental/VirtualScouting/VirtualScouting.uplugin
--- a/Engine/Plugins/Experimental/VirtualScouting/VirtualScouting.uplugin
+++ b/Engine/Plugins/Experimental/VirtualScouting/VirtualScouting.uplugin
@@ -70,7 +70,7 @@
 			"Enabled": true
 		},
 		{
-			"Name": "SunPosition",
+			"Name": "LiveLink",
 			"Enabled": true
 		},
 		{
