Virtual Scouting - removed Delete actors  library function  as can be completed using BP functions in TypedElement framework
Better to have users see a BP example of typed element list management and delete without ref checks

#virtualized

diff --git a/Engine/Plugins/Experimental/VirtualScouting/Source/VirtualScoutingEditor/Private/VirtualScoutingBlueprintFunctionLibrary.cpp b/Engine/Plugins/Experimental/VirtualScouting/Source/VirtualScoutingEditor/Private/VirtualScoutingBlueprintFunctionLibrary.cpp
--- a/Engine/Plugins/Experimental/VirtualScouting/Source/VirtualScoutingEditor/Private/VirtualScoutingBlueprintFunctionLibrary.cpp
+++ b/Engine/Plugins/Experimental/VirtualScouting/Source/VirtualScoutingEditor/Private/VirtualScoutingBlueprintFunctionLibrary.cpp
@@ -5,9 +5,6 @@
 
 #include "Editor.h"
 #include "ScopedTransaction.h"
-#include "Editor/UnrealEdEngine.h"
-#include "Subsystems/EditorActorSubsystem.h"
-#include "Subsystems/UnrealEditorSubsystem.h"
 
 bool UVirtualScoutingBlueprintFunctionLibrary::CheckIsWithEditor()
 {
@@ -17,37 +14,3 @@ bool UVirtualScoutingBlueprintFunctionLibrary::CheckIsWithEditor()
 	return false;
 #endif
 }
-
-bool UVirtualScoutingBlueprintFunctionLibrary::DeleteActors(const TArray<AActor*>& InActorsToDelete)
-{
-	const FScopedTransaction Transaction(NSLOCTEXT("VirtualScouting", "DeleteActors", "Delete Actors"));
-
-	UWorld* World = nullptr;
-	UUnrealEditorSubsystem* UnrealEditorSubsystem = GEditor->GetEditorSubsystem<UUnrealEditorSubsystem>();
-	World = UnrealEditorSubsystem->GetEditorWorld();
-	if (World)
-	{
-		// select
-		UEditorActorSubsystem* EditorActorSubsystem = GEditor->GetEditorSubsystem<UEditorActorSubsystem>();
-		if (EditorActorSubsystem)
-		{
-			EditorActorSubsystem->SetSelectedLevelActors(InActorsToDelete);
-		}
-		// delete selected
-		const bool bVerifyDeletionCanHappen = true;
-		const bool bWarnAboutReferences = false;
-		return GEditor->edactDeleteSelected(World, bVerifyDeletionCanHappen, bWarnAboutReferences, bWarnAboutReferences);
-	}
-
-	// If EditorWorld is not valid (i.e. VRPreviewMode), just call DestroyActor
-	bool bDeleted = false;
-	for (AActor* Actor : InActorsToDelete)
-	{
-		World = Actor->GetWorld();
-		if (World)
-		{
-			bDeleted |= World->DestroyActor(Actor);
-		}
-	}
-	return bDeleted;
-}
\ No newline at end of file
diff --git a/Engine/Plugins/Experimental/VirtualScouting/Source/VirtualScoutingEditor/Public/VirtualScoutingBlueprintFunctionLibrary.h b/Engine/Plugins/Experimental/VirtualScouting/Source/VirtualScoutingEditor/Public/VirtualScoutingBlueprintFunctionLibrary.h
--- a/Engine/Plugins/Experimental/VirtualScouting/Source/VirtualScoutingEditor/Public/VirtualScoutingBlueprintFunctionLibrary.h
+++ b/Engine/Plugins/Experimental/VirtualScouting/Source/VirtualScoutingEditor/Public/VirtualScoutingBlueprintFunctionLibrary.h
@@ -17,8 +17,4 @@ class VIRTUALSCOUTINGEDITOR_API UVirtualScoutingBlueprintFunctionLibrary : publi
 	UFUNCTION(BlueprintCallable, Category="VirtualScouting")
 	static bool CheckIsWithEditor();
 
-	/** Helper function to delete actors. This skips reference checks.*/
-	UFUNCTION(BlueprintCallable, Category="VirtualScouting")
-	static bool DeleteActors(const TArray<AActor*>& InActorsToDelete);
-
 };
