Virtual Scouting - WIP improved measurement tool and menu now using niagara for VFX
-Added dependency to VPUtils back into the plugin.

#virtualized

diff --git a/Engine/Plugins/Experimental/VirtualScouting/VirtualScouting.uplugin b/Engine/Plugins/Experimental/VirtualScouting/VirtualScouting.uplugin
--- a/Engine/Plugins/Experimental/VirtualScouting/VirtualScouting.uplugin
+++ b/Engine/Plugins/Experimental/VirtualScouting/VirtualScouting.uplugin
@@ -104,6 +104,10 @@
 		{
 			"Name": "CineCameraSceneCapture",
 			"Enabled": true
+		},
+		{
+			"Name": "VirtualProductionUtilities",
+			"Enabled": true
 		}
 	]
 }
