Virtual Scouting - add persistent bUseSmooth rotation and SequenceToolCollection Fname variables to scouting settings

diff --git a/Engine/Plugins/Experimental/VirtualScouting/Source/VirtualScouting/Public/VirtualScoutingSettings.h b/Engine/Plugins/Experimental/VirtualScouting/Source/VirtualScouting/Public/VirtualScoutingSettings.h
--- a/Engine/Plugins/Experimental/VirtualScouting/Source/VirtualScouting/Public/VirtualScoutingSettings.h
+++ b/Engine/Plugins/Experimental/VirtualScouting/Source/VirtualScouting/Public/VirtualScoutingSettings.h
@@ -25,6 +25,12 @@ public:
 	UPROPERTY(Config, EditAnywhere, BlueprintReadWrite, Category="Virtual Scouting", meta=(DisplayName="Viewfinder Monitor Masks"))
 	TArray<float> ViewfinderMaskArray = {1.33, 1.66, 1.78, 2.0, 2.35, 2.39 };
 
+	UPROPERTY(Config, EditAnywhere, BlueprintReadWrite, Category="Virtual Scouting", meta=(DisplayName="Smooth Rotation Movement"))
+	bool UseSmoothRotation = false;
+
+	UPROPERTY(Config, EditAnywhere, BlueprintReadWrite, Category="Virtual Scouting", meta=(DisplayName="Sequence Tool Collection"))
+	FName SequenceToolCollection;
+	
 	UFUNCTION(BlueprintPure, Category="Virtual Scouting", DisplayName="Virtual Scouting Settings")
 	static UVirtualScoutingSettings* GetVirtualScoutingSettings();
 };
