Virtual Scouting - added Virtual Scouting plugin project setting for viewfinder autoexposure use and exposure compensation

diff --git a/Engine/Plugins/Experimental/VirtualScouting/Source/VirtualScouting/Public/VirtualScoutingSettings.h b/Engine/Plugins/Experimental/VirtualScouting/Source/VirtualScouting/Public/VirtualScoutingSettings.h
--- a/Engine/Plugins/Experimental/VirtualScouting/Source/VirtualScouting/Public/VirtualScoutingSettings.h
+++ b/Engine/Plugins/Experimental/VirtualScouting/Source/VirtualScouting/Public/VirtualScoutingSettings.h
@@ -19,6 +19,12 @@ public:
 	UPROPERTY(Config, EditAnywhere, BlueprintReadWrite, Category="Virtual Scouting", meta=(DisplayName="Show Measurements in Imperial Units"))
 	bool bUseImperial = false;
 
+	UPROPERTY(Config, EditAnywhere, BlueprintReadWrite, Category="Virtual Scouting", meta=(DisplayName="Viewfinder Use AutoExposure"))
+	bool bViewfinderUseExposure = false;
+
+	UPROPERTY(Config, EditAnywhere, BlueprintReadWrite, Category="Virtual Scouting", meta=(DisplayName="Viewfinder ExposureCompensation", ClampMin=-15, ClampMax=15))
+	float ViewfinderExposureCompensation = 1;
+	
 	UPROPERTY(Config, EditAnywhere, BlueprintReadWrite, Category="Virtual Scouting", meta=(DisplayName="Viewfinder Apertures"))
 	TArray<float> ViewfinderApertureArray = {1.2, 2.0, 2.8, 4.0, 5.6, 8.0, 11.0, 16.0, 22.0};
 
