Virtual Scouting: Remove cross-module Private directory reference.

Also removes redundant `bTreatAsEngineModule = true` holdovers from when these were iterated on as sandbox project plugins.

#rnx

diff --git a/Engine/Plugins/Experimental/VirtualScouting/Source/VirtualScouting/VirtualScouting.Build.cs b/Engine/Plugins/Experimental/VirtualScouting/Source/VirtualScouting/VirtualScouting.Build.cs
--- a/Engine/Plugins/Experimental/VirtualScouting/Source/VirtualScouting/VirtualScouting.Build.cs
+++ b/Engine/Plugins/Experimental/VirtualScouting/Source/VirtualScouting/VirtualScouting.Build.cs
@@ -7,7 +7,6 @@ public class VirtualScouting : ModuleRules
 	public VirtualScouting(ReadOnlyTargetRules Target) : base(Target)
 	{
 		PCHUsage = ModuleRules.PCHUsageMode.UseExplicitOrSharedPCHs;
-		bTreatAsEngineModule = true;
 
 		PublicDependencyModuleNames.AddRange(
 			new string[]
diff --git a/Engine/Plugins/Experimental/VirtualScouting/Source/VirtualScoutingEditor/VirtualScoutingEditor.Build.cs b/Engine/Plugins/Experimental/VirtualScouting/Source/VirtualScoutingEditor/VirtualScoutingEditor.Build.cs
--- a/Engine/Plugins/Experimental/VirtualScouting/Source/VirtualScoutingEditor/VirtualScoutingEditor.Build.cs
+++ b/Engine/Plugins/Experimental/VirtualScouting/Source/VirtualScoutingEditor/VirtualScoutingEditor.Build.cs
@@ -7,14 +7,6 @@ public class VirtualScoutingEditor : ModuleRules
 	public VirtualScoutingEditor(ReadOnlyTargetRules Target) : base(Target)
 	{
 		PCHUsage = ModuleRules.PCHUsageMode.UseExplicitOrSharedPCHs;
-		bTreatAsEngineModule = true;
-
-		PrivateIncludePaths.AddRange(
-			new string[]
-			{
-				"VirtualScoutingOpenXR/Private",
-			}
-		);
 
 		PublicDependencyModuleNames.AddRange(new string[]
 			{
diff --git a/Engine/Plugins/Experimental/VirtualScouting/Source/VirtualScoutingOpenXR/Internal/VirtualScoutingOpenXR.h b/Engine/Plugins/Experimental/VirtualScouting/Source/VirtualScoutingOpenXR/Internal/VirtualScoutingOpenXR.h
new file mode 100644
--- /dev/null
+++ b/Engine/Plugins/Experimental/VirtualScouting/Source/VirtualScoutingOpenXR/Internal/VirtualScoutingOpenXR.h
@@ -0,0 +1,66 @@
+// Copyright Epic Games, Inc. All Rights Reserved.
+
+#pragma once
+
+#if WITH_EDITOR
+
+#include "IOpenXRExtensionPlugin.h"
+#include "Async/Future.h"
+#include <openxr/openxr_platform_defines.h>
+
+template <typename OptionalType> struct TOptional;
+
+
+class VIRTUALSCOUTINGOPENXR_API FVirtualScoutingOpenXRExtension
+	: public IOpenXRExtensionPlugin
+	, public TSharedFromThis<FVirtualScoutingOpenXRExtension>
+{
+public:
+	FVirtualScoutingOpenXRExtension();
+	virtual ~FVirtualScoutingOpenXRExtension();
+
+	//~ Begin IOpenXRExtensionPlugin interface
+	virtual FString GetDisplayName() override
+	{
+		return FString(TEXT("VirtualScouting"));
+	}
+
+	virtual bool GetOptionalExtensions(TArray<const ANSICHAR*>& OutExtensions) override;
+	virtual void OnEvent(XrSession InSession, const XrEventDataBaseHeader* InHeader) override;
+	virtual void PostCreateInstance(XrInstance InInstance) override;
+	virtual void PostCreateSession(XrSession InSession) override;
+	virtual void PostSyncActions(XrSession InSession) override;
+	//~ End IOpenXRExtensionPlugin interface
+
+	TFuture<FName>& GetHmdDeviceTypeFuture() { return DeviceTypeFuture; }
+
+private:
+	void OnVREditingModeEnter();
+	void OnVREditingModeExit();
+
+	TOptional<FName> TryGetHmdDeviceType();
+	void TryFulfillDeviceTypePromise();
+
+private:
+	FDelegateHandle InitCompleteDelegate;
+
+	XrDebugUtilsMessengerEXT Messenger = XR_NULL_HANDLE;
+
+	XrInstance Instance = XR_NULL_HANDLE;
+	XrSession Session = XR_NULL_HANDLE;
+	XrActionSet ActionSet = XR_NULL_HANDLE;
+
+	bool bIsVrEditingModeActive = false;
+
+	TPromise<FName> DeviceTypePromise;
+	TFuture<FName> DeviceTypeFuture;
+
+private:
+	static XrBool32 XRAPI_CALL XrDebugUtilsMessengerCallback_Trampoline(
+		XrDebugUtilsMessageSeverityFlagsEXT InMessageSeverity,
+		XrDebugUtilsMessageTypeFlagsEXT InMessageTypes,
+		const XrDebugUtilsMessengerCallbackDataEXT* InCallbackData,
+		void* InUserData);
+};
+
+#endif // #if WITH_EDITOR
diff --git a/Engine/Plugins/Experimental/VirtualScouting/Source/VirtualScoutingOpenXR/Internal/VirtualScoutingOpenXRModule.h b/Engine/Plugins/Experimental/VirtualScouting/Source/VirtualScoutingOpenXR/Internal/VirtualScoutingOpenXRModule.h
new file mode 100644
--- /dev/null
+++ b/Engine/Plugins/Experimental/VirtualScouting/Source/VirtualScoutingOpenXR/Internal/VirtualScoutingOpenXRModule.h
@@ -0,0 +1,33 @@
+// Copyright Epic Games, Inc. All Rights Reserved.
+
+#pragma once
+
+#include "Modules/ModuleManager.h"
+
+
+class FVirtualScoutingOpenXRExtension;
+
+
+class VIRTUALSCOUTINGOPENXR_API FVirtualScoutingOpenXRModule : public IModuleInterface
+{
+public:
+	/**
+	 * Singleton-like access to this module's interface.  This is just for convenience!
+	 * Beware of calling this during the shutdown phase, though.  Your module might have been unloaded already.
+	 *
+	 * @return Returns singleton instance, loading the module on demand if needed
+	 */
+	static FVirtualScoutingOpenXRModule& Get()
+	{
+		static const FName ModuleName = "VirtualScoutingOpenXR";
+		return FModuleManager::LoadModuleChecked<FVirtualScoutingOpenXRModule>(ModuleName);
+	};
+
+	virtual void StartupModule() override;
+	virtual void ShutdownModule() override;
+
+	const TSharedPtr<FVirtualScoutingOpenXRExtension>& GetOpenXRExt() { return OpenXRExt; }
+
+private:
+	TSharedPtr<FVirtualScoutingOpenXRExtension> OpenXRExt;
+};
diff --git a/Engine/Plugins/Experimental/VirtualScouting/Source/VirtualScoutingOpenXR/Private/VirtualScoutingOpenXR.h b/Engine/Plugins/Experimental/VirtualScouting/Source/VirtualScoutingOpenXR/Private/VirtualScoutingOpenXR.h
deleted file mode 100644
diff --git a/Engine/Plugins/Experimental/VirtualScouting/Source/VirtualScoutingOpenXR/Private/VirtualScoutingOpenXRModule.h b/Engine/Plugins/Experimental/VirtualScouting/Source/VirtualScoutingOpenXR/Private/VirtualScoutingOpenXRModule.h
deleted file mode 100644
diff --git a/Engine/Plugins/Experimental/VirtualScouting/Source/VirtualScoutingOpenXR/VirtualScoutingOpenXR.Build.cs b/Engine/Plugins/Experimental/VirtualScouting/Source/VirtualScoutingOpenXR/VirtualScoutingOpenXR.Build.cs
--- a/Engine/Plugins/Experimental/VirtualScouting/Source/VirtualScoutingOpenXR/VirtualScoutingOpenXR.Build.cs
+++ b/Engine/Plugins/Experimental/VirtualScouting/Source/VirtualScoutingOpenXR/VirtualScoutingOpenXR.Build.cs
@@ -7,7 +7,6 @@ public class VirtualScoutingOpenXR : ModuleRules
 	public VirtualScoutingOpenXR(ReadOnlyTargetRules Target) : base(Target)
 	{
 		PCHUsage = ModuleRules.PCHUsageMode.UseExplicitOrSharedPCHs;
-		bTreatAsEngineModule = true;
 
 		PublicDependencyModuleNames.AddRange(new string[]
 			{
diff --git a/Engine/Plugins/Experimental/XRCreativeFramework/Source/XRCreative/XRCreative.Build.cs b/Engine/Plugins/Experimental/XRCreativeFramework/Source/XRCreative/XRCreative.Build.cs
--- a/Engine/Plugins/Experimental/XRCreativeFramework/Source/XRCreative/XRCreative.Build.cs
+++ b/Engine/Plugins/Experimental/XRCreativeFramework/Source/XRCreative/XRCreative.Build.cs
@@ -7,7 +7,6 @@ public class XRCreative : ModuleRules
 	public XRCreative(ReadOnlyTargetRules Target) : base(Target)
 	{
 		PCHUsage = ModuleRules.PCHUsageMode.UseExplicitOrSharedPCHs;
-		bTreatAsEngineModule = true;
 
 		PublicDependencyModuleNames.AddRange(
 			new string[]
diff --git a/Engine/Plugins/Experimental/XRCreativeFramework/Source/XRCreativeEditor/XRCreativeEditor.Build.cs b/Engine/Plugins/Experimental/XRCreativeFramework/Source/XRCreativeEditor/XRCreativeEditor.Build.cs
--- a/Engine/Plugins/Experimental/XRCreativeFramework/Source/XRCreativeEditor/XRCreativeEditor.Build.cs
+++ b/Engine/Plugins/Experimental/XRCreativeFramework/Source/XRCreativeEditor/XRCreativeEditor.Build.cs
@@ -7,7 +7,6 @@ public class XRCreativeEditor : ModuleRules
 	public XRCreativeEditor(ReadOnlyTargetRules Target) : base(Target)
 	{
 		PCHUsage = ModuleRules.PCHUsageMode.UseExplicitOrSharedPCHs;
-		bTreatAsEngineModule = true;
 
 		PublicDependencyModuleNames.AddRange(new string[]
 			{
