Virtual Scouting - User feedback: added movement settings to PerUserEditor settings so they can persist between VR sessions and engine sessions

diff --git a/Engine/Plugins/Experimental/VirtualScouting/Source/VirtualScouting/Public/VirtualScoutingSettings.h b/Engine/Plugins/Experimental/VirtualScouting/Source/VirtualScouting/Public/VirtualScoutingSettings.h
--- a/Engine/Plugins/Experimental/VirtualScouting/Source/VirtualScouting/Public/VirtualScoutingSettings.h
+++ b/Engine/Plugins/Experimental/VirtualScouting/Source/VirtualScouting/Public/VirtualScoutingSettings.h
@@ -31,9 +31,6 @@ public:
 	UPROPERTY(Config, EditAnywhere, BlueprintReadWrite, Category="Virtual Scouting", meta=(DisplayName="Viewfinder Monitor Masks"))
 	TArray<float> ViewfinderMaskArray = {1.33, 1.66, 1.78, 2.0, 2.35, 2.39 };
 
-	UPROPERTY(Config, EditAnywhere, BlueprintReadWrite, Category="Virtual Scouting", meta=(DisplayName="Smooth Rotation Movement"))
-	bool UseSmoothRotation = false;
-
 	UPROPERTY(Config, EditAnywhere, BlueprintReadWrite, Category="Virtual Scouting", meta=(DisplayName="Sequence Tool Collection"))
 	FName SequenceToolCollection;
 	
@@ -53,4 +50,25 @@ public:
 
 	UFUNCTION(BlueprintPure, Category="Virtual Scouting Editor")
 	static UVirtualScoutingEditorSettings* GetVirtualScoutingEditorSettings();
+
+	UPROPERTY(Config, EditAnywhere, BlueprintReadWrite, Category="Virtual Scouting User", meta=(DisplayName= "Flight Speed", ToolTip="Avatar flying speed. Default is 4.0", ClampMin=1.0f, ClampMax=10.0f));
+	float FlightSpeed = 4.0f;
+
+	UPROPERTY(Config, EditAnywhere, BlueprintReadWrite, Category="Virtual Scouting User", meta=(DisplayName= "Drag Speed", ToolTip="Speed of movement when you drag-move. Default is 0.7", ClampMin=0.1f, ClampMax=2.0f));
+	float DragSpeed = 0.7f;
+
+	UPROPERTY(Config, EditAnywhere, BlueprintReadWrite, Category="Virtual Scouting User", meta=(DisplayName= "Show Tooltips", ToolTip="Show Tooltips when in VR. These appear when motioncontroller is brought near to HMD"));
+	bool bEnableTooltips = true;
+
+	UPROPERTY(Config, EditAnywhere, BlueprintReadWrite, Category="Virtual Scouting User", meta=(DisplayName= "Use Smooth Rotation", ToolTip="True = Rotate smoothly. False = Flick Rotate. Default is Flick Rotate"));
+	bool bUseSmoothRotation = false;
+
+	UPROPERTY(Config, EditAnywhere, BlueprintReadWrite, Category="Virtual Scouting User", meta=(DisplayName= "Use Teleport Rotation", ToolTip="Use the forward axis roll from the motion controller to define and adjust teleport rotation"));
+	bool bUseTeleportRotation = false;
+
+	UPROPERTY(Config, EditAnywhere, BlueprintReadWrite, Category="Virtual Scouting User", meta=(DisplayName= "Teleport Rotation Threshold", ToolTip="How many degrees of rotation of the Motion Controller before the teleporter rotation is modified. Default is 22 degrees"));
+	float TeleportRotationThreshold = 22.0f;
+
+	
+	
 };
