Virtual Scouting - removed teleport rotation from Virtual Scouting Editor settings as not needed.

diff --git a/Engine/Plugins/Experimental/VirtualScouting/Source/VirtualScouting/Public/VirtualScoutingSettings.h b/Engine/Plugins/Experimental/VirtualScouting/Source/VirtualScouting/Public/VirtualScoutingSettings.h
--- a/Engine/Plugins/Experimental/VirtualScouting/Source/VirtualScouting/Public/VirtualScoutingSettings.h
+++ b/Engine/Plugins/Experimental/VirtualScouting/Source/VirtualScouting/Public/VirtualScoutingSettings.h
@@ -65,10 +65,5 @@ public:
 
 	UPROPERTY(Config, EditAnywhere, BlueprintReadWrite, Category="Virtual Scouting User", meta=(DisplayName= "Use Teleport Rotation", ToolTip="Use the forward axis roll from the motion controller to define and adjust teleport rotation"));
 	bool bUseTeleportRotation = false;
-
-	UPROPERTY(Config, EditAnywhere, BlueprintReadWrite, Category="Virtual Scouting User", meta=(DisplayName= "Teleport Rotation Threshold", ToolTip="How many degrees of rotation of the Motion Controller before the teleporter rotation is modified. Default is 22 degrees"));
-	float TeleportRotationThreshold = 22.0f;
-
-	
 	
 };
