Virtual Scouting: Ensure we default to "Scouting Default" mode when enabling the plugin for the first time.

#rnx

diff --git a/Engine/Plugins/Experimental/VirtualScouting/Source/VirtualScoutingEditor/Private/VirtualScoutingEditorModule.cpp b/Engine/Plugins/Experimental/VirtualScouting/Source/VirtualScoutingEditor/Private/VirtualScoutingEditorModule.cpp
--- a/Engine/Plugins/Experimental/VirtualScouting/Source/VirtualScoutingEditor/Private/VirtualScoutingEditorModule.cpp
+++ b/Engine/Plugins/Experimental/VirtualScouting/Source/VirtualScoutingEditor/Private/VirtualScoutingEditorModule.cpp
@@ -2,6 +2,7 @@
 
 #include "VirtualScoutingEditorModule.h"
 #include "Modules/ModuleManager.h"
+#include "VRModeSettings.h"
 
 
 DEFINE_LOG_CATEGORY(LogVirtualScoutingEditor);
@@ -12,6 +13,13 @@ class FVirtualScoutingEditorModule : public IModuleInterface
 private:
 	virtual void StartupModule() override
 	{
+		// When the user has no previous selection upon enabling this plugin
+		// for the first time, ensure we default to our preferred mode.
+		UVRModeSettings* Settings = GetMutableDefault<UVRModeSettings>();
+		if (Settings->ModeClass.IsNull())
+		{
+			Settings->ModeClass = FSoftObjectPath("/VirtualScouting/Core/Scouting_Default.Scouting_Default_C");
+		}
 	}
 
 	virtual void ShutdownModule() override
