Virtual Scout: Added Content Browser Content Filter settings to Project Settings

diff --git a/Engine/Plugins/Experimental/VirtualScouting/Source/VirtualScouting/Public/VirtualScoutingSettings.h b/Engine/Plugins/Experimental/VirtualScouting/Source/VirtualScouting/Public/VirtualScoutingSettings.h
--- a/Engine/Plugins/Experimental/VirtualScouting/Source/VirtualScouting/Public/VirtualScoutingSettings.h
+++ b/Engine/Plugins/Experimental/VirtualScouting/Source/VirtualScouting/Public/VirtualScoutingSettings.h
@@ -33,6 +33,9 @@ public:
 
 	UPROPERTY(Config, EditAnywhere, BlueprintReadWrite, Category="Virtual Scouting", meta=(DisplayName="Sequence Tool Collection"))
 	FName SequenceToolCollection;
+
+	UPROPERTY(Config, EditAnywhere, BlueprintReadWrite, Category = "Virtual Scouting", meta = (DisplayName = "Content Browser Tool Collection"))
+	FName ContentBrowserToolCollection;
 	
 	UFUNCTION(BlueprintPure, Category="Virtual Scouting", DisplayName="Virtual Scouting Settings")
 	static UVirtualScoutingSettings* GetVirtualScoutingSettings();
