XR Creative / Virtual Scouting: Remove dependency on Settings module, use UDeveloperSettings. Enables packaging in Shipping configuration.

diff --git a/Engine/Plugins/Experimental/VirtualScouting/Source/VirtualScouting/Private/Module.cpp b/Engine/Plugins/Experimental/VirtualScouting/Source/VirtualScouting/Private/Module.cpp
--- a/Engine/Plugins/Experimental/VirtualScouting/Source/VirtualScouting/Private/Module.cpp
+++ b/Engine/Plugins/Experimental/VirtualScouting/Source/VirtualScouting/Private/Module.cpp
@@ -2,9 +2,6 @@
 
 #include "Module.h"
 #include "VirtualScoutingLog.h"
-#include "VirtualScoutingSettings.h"
-#include "ISettingsModule.h"
-#include "ISettingsSection.h"
 
 
 #define LOCTEXT_NAMESPACE "FVirtualScoutingModule"
@@ -15,55 +12,11 @@ DEFINE_LOG_CATEGORY(LogVirtualScouting);
 
 void FVirtualScoutingModule::StartupModule()
 {
-	if(ISettingsModule* SettingsModule = FModuleManager::GetModulePtr<ISettingsModule>("Settings"))
-	{
-		SettingsModule->RegisterSettings("Project", "Plugins", "VirtualScoutingSettings",
-			LOCTEXT("RuntimeSettingsName", "Virtual Scouting Settings"), LOCTEXT("RuntimeSettingsDescription", "Virtual Scouting Settings"),
-			GetMutableDefault<UVirtualScoutingSettings>());
-
-		SettingsModule->RegisterSettings("Editor", "Plugins", "VirtualScoutingEditorSettings",
-			LOCTEXT("EditorSettingsName", "Virtual Scouting Editor Settings"), LOCTEXT("EditorSettingsDescription", "Virtual Scouting Editor Settings"),
-			GetMutableDefault<UVirtualScoutingEditorSettings>());
-	}
 }
 
 
 void FVirtualScoutingModule::ShutdownModule()
 {
-	if(ISettingsModule* SettingsModule = FModuleManager::GetModulePtr<ISettingsModule>("Settings"))
-	{
-		SettingsModule->UnregisterSettings("Project", "Plugins", "VirtualScoutingSettings");
-		SettingsModule->UnregisterSettings("Editor", "Plugins", "VirtualScoutingEditorSettings");
-	}
-}
-
-
-void FVirtualScoutingModule::RegisterSettings() const
-{
-	ISettingsModule* SettingsModule = FModuleManager::GetModulePtr<ISettingsModule>("Settings");
-	if (SettingsModule != nullptr)
-	{
-		ISettingsSectionPtr SettingsSection = SettingsModule->RegisterSettings("Project", "Plugins", "VirtualScouting",
-			LOCTEXT("VirtualScoutingSettingsName", "Virtual Scouting"),
-			LOCTEXT("VirtualScoutingSettingsDescription", "Configure the Virtual Scouting settings."),
-			GetMutableDefault<UVirtualScoutingSettings>());
-	}
-}
-
-
-void FVirtualScoutingModule::UnregisterSettings() const
-{
-	ISettingsModule* SettingsModule = FModuleManager::GetModulePtr<ISettingsModule>("Settings");
-	if (SettingsModule != nullptr)
-	{
-		SettingsModule->UnregisterSettings("Project", "Plugins", "VirtualScouting");
-	}
-}
-
-
-bool FVirtualScoutingModule::OnSettingsModified()
-{
-	return true;
 }
 
 
diff --git a/Engine/Plugins/Experimental/VirtualScouting/Source/VirtualScouting/Private/Module.h b/Engine/Plugins/Experimental/VirtualScouting/Source/VirtualScouting/Private/Module.h
--- a/Engine/Plugins/Experimental/VirtualScouting/Source/VirtualScouting/Private/Module.h
+++ b/Engine/Plugins/Experimental/VirtualScouting/Source/VirtualScouting/Private/Module.h
@@ -2,7 +2,6 @@
 
 #pragma once
 
-#include "CoreMinimal.h"
 #include "Modules/ModuleManager.h"
 
 
@@ -11,14 +10,4 @@ class FVirtualScoutingModule : public IModuleInterface
 public:
 	virtual void StartupModule() override;
 	virtual void ShutdownModule() override;
-
-private:
-	/** Register the settings class with the settings module **/
-	void RegisterSettings() const;
-
-	/** Unregister the settings class with the settings module **/
-	void UnregisterSettings() const;
-	
-	/** TODO Handler for when Virtual Scouting settings are changed. */
-	bool OnSettingsModified();
 };
diff --git a/Engine/Plugins/Experimental/VirtualScouting/Source/VirtualScouting/Private/VirtualScoutingSettings.cpp b/Engine/Plugins/Experimental/VirtualScouting/Source/VirtualScouting/Private/VirtualScoutingSettings.cpp
--- a/Engine/Plugins/Experimental/VirtualScouting/Source/VirtualScouting/Private/VirtualScoutingSettings.cpp
+++ b/Engine/Plugins/Experimental/VirtualScouting/Source/VirtualScouting/Private/VirtualScoutingSettings.cpp
@@ -3,11 +3,27 @@
 #include "VirtualScoutingSettings.h"
 
 
+namespace UE::VirtualScouting::Private
+{
+	static const FName PluginsCategoryName = TEXT("Plugins");
+};
+
+
+FName UVirtualScoutingSettings::GetCategoryName() const
+{
+	return UE::VirtualScouting::Private::PluginsCategoryName;
+}
+
 UVirtualScoutingSettings* UVirtualScoutingSettings::GetVirtualScoutingSettings()
 {
 	return GetMutableDefault<UVirtualScoutingSettings>();
 }
 
+FName UVirtualScoutingEditorSettings::GetCategoryName() const
+{
+	return UE::VirtualScouting::Private::PluginsCategoryName;
+}
+
 UVirtualScoutingEditorSettings* UVirtualScoutingEditorSettings::GetVirtualScoutingEditorSettings()
 {
 	return GetMutableDefault<UVirtualScoutingEditorSettings>();
diff --git a/Engine/Plugins/Experimental/VirtualScouting/Source/VirtualScouting/Public/VirtualScoutingSettings.h b/Engine/Plugins/Experimental/VirtualScouting/Source/VirtualScouting/Public/VirtualScoutingSettings.h
--- a/Engine/Plugins/Experimental/VirtualScouting/Source/VirtualScouting/Public/VirtualScoutingSettings.h
+++ b/Engine/Plugins/Experimental/VirtualScouting/Source/VirtualScouting/Public/VirtualScoutingSettings.h
@@ -2,7 +2,7 @@
 
 #pragma once
 
-#include "Engine/DeveloperSettingsBackedByCVars.h"
+#include "Engine/DeveloperSettings.h"
 #include "VirtualScoutingSettings.generated.h"
 
 
@@ -10,11 +10,17 @@
  * Per project settings for Virtual Scouting.
  */
 UCLASS(Config=VirtualScoutingSettings, DefaultConfig, DisplayName="Virtual Scouting")
-class VIRTUALSCOUTING_API UVirtualScoutingSettings : public UObject
+class VIRTUALSCOUTING_API UVirtualScoutingSettings : public UDeveloperSettings
 {
 	GENERATED_BODY()
 	
 public:
+	//~ Begin UDeveloperSettings interface
+	virtual FName GetCategoryName() const override;
+	//~ End UDeveloperSettings interface
+
+	UFUNCTION(BlueprintPure, Category="Virtual Scouting", DisplayName="Virtual Scouting Settings")
+	static UVirtualScoutingSettings* GetVirtualScoutingSettings();
 
 	UPROPERTY(Config, EditAnywhere, BlueprintReadWrite, Category="Virtual Scouting", meta=(DisplayName="Show Measurements in Imperial Units"))
 	bool bUseImperial = false;
@@ -36,20 +42,20 @@ public:
 
 	UPROPERTY(Config, EditAnywhere, BlueprintReadWrite, Category = "Virtual Scouting", meta = (DisplayName = "Content Browser Tool Collection"))
 	FName ContentBrowserToolCollection;
-	
-	UFUNCTION(BlueprintPure, Category="Virtual Scouting", DisplayName="Virtual Scouting Settings")
-	static UVirtualScoutingSettings* GetVirtualScoutingSettings();
 };
 
 /**
  * Per user settings for Virtual Scouting Editor.
  */
-UCLASS(Config=EditorPerProjectUserSettings, DisplayName="Virtual Scouting Editor Settings")
-class VIRTUALSCOUTING_API UVirtualScoutingEditorSettings : public UObject
+UCLASS(Config=EditorPerProjectUserSettings, DisplayName="Virtual Scouting Editor")
+class VIRTUALSCOUTING_API UVirtualScoutingEditorSettings : public UDeveloperSettings
 {
 	GENERATED_BODY()
 	
 public:
+	//~ Begin UDeveloperSettings interface
+	virtual FName GetCategoryName() const override;
+	//~ End UDeveloperSettings interface
 
 	UFUNCTION(BlueprintPure, Category="Virtual Scouting Editor")
 	static UVirtualScoutingEditorSettings* GetVirtualScoutingEditorSettings();
@@ -68,5 +74,4 @@ public:
 
 	UPROPERTY(Config, EditAnywhere, BlueprintReadWrite, Category="Virtual Scouting User", meta=(DisplayName= "Use Teleport Rotation", ToolTip="Use the forward axis roll from the motion controller to define and adjust teleport rotation"));
 	bool bUseTeleportRotation = false;
-	
 };
diff --git a/Engine/Plugins/Experimental/VirtualScouting/Source/VirtualScouting/VirtualScouting.Build.cs b/Engine/Plugins/Experimental/VirtualScouting/Source/VirtualScouting/VirtualScouting.Build.cs
--- a/Engine/Plugins/Experimental/VirtualScouting/Source/VirtualScouting/VirtualScouting.Build.cs
+++ b/Engine/Plugins/Experimental/VirtualScouting/Source/VirtualScouting/VirtualScouting.Build.cs
@@ -29,7 +29,6 @@ public class VirtualScouting : ModuleRules
 				"InteractiveToolsFramework",
 				"LevelSequence",
 				"RenderCore", // TODO: Kill this? Only for FlushRenderingCommands
-				"Settings",
 				"Slate",
 				"SlateCore",
 				"TypedElementFramework",
diff --git a/Engine/Plugins/Experimental/XRCreativeFramework/Source/XRCreative/Private/Module.cpp b/Engine/Plugins/Experimental/XRCreativeFramework/Source/XRCreative/Private/Module.cpp
--- a/Engine/Plugins/Experimental/XRCreativeFramework/Source/XRCreative/Private/Module.cpp
+++ b/Engine/Plugins/Experimental/XRCreativeFramework/Source/XRCreative/Private/Module.cpp
@@ -1,9 +1,6 @@
 // Copyright Epic Games, Inc. All Rights Reserved.
 
 #include "Module.h"
-#include "XRCreativeSettings.h"
-#include "ISettingsModule.h"
-#include "ISettingsSection.h"
 #include "XRCreativeLog.h"
 
 
@@ -15,27 +12,11 @@ DEFINE_LOG_CATEGORY(LogXRCreative);
 
 void FXRCreativeModule::StartupModule()
 {
-	if(ISettingsModule* SettingsModule = FModuleManager::GetModulePtr<ISettingsModule>("Settings"))
-	{
-		SettingsModule->RegisterSettings("Project", "Plugins", "XRCreativeSettings",
-			LOCTEXT("RuntimeSettingsName", "XRCreative Settings"), LOCTEXT("RuntimeSettingsDescription", "XRCreative Settings"),
-			GetMutableDefault<UXRCreativeSettings>());
-
-		SettingsModule->RegisterSettings("Editor", "Plugins", "XRCreativeEditorSettings",
-			LOCTEXT("EditorSettingsName", "XRCreative Editor Settings"), LOCTEXT("EditorSettingsDescription", "XRCreative Editor Settings"),
-			GetMutableDefault<UXRCreativeEditorSettings>());
-	}
 }
 
 
 void FXRCreativeModule::ShutdownModule()
 {
-
-	if(ISettingsModule* SettingsModule = FModuleManager::GetModulePtr<ISettingsModule>("Settings"))
-	{
-		SettingsModule->UnregisterSettings("Project", "Plugins", "XRCreativeSettings");
-		SettingsModule->UnregisterSettings("Editor", "Plugins", "XRCreativeEditorSettings");
-	}
 }
 
 
diff --git a/Engine/Plugins/Experimental/XRCreativeFramework/Source/XRCreative/Private/XRCreativeSettings.cpp b/Engine/Plugins/Experimental/XRCreativeFramework/Source/XRCreative/Private/XRCreativeSettings.cpp
--- a/Engine/Plugins/Experimental/XRCreativeFramework/Source/XRCreative/Private/XRCreativeSettings.cpp
+++ b/Engine/Plugins/Experimental/XRCreativeFramework/Source/XRCreative/Private/XRCreativeSettings.cpp
@@ -3,11 +3,27 @@
 #include "XRCreativeSettings.h"
 
 
+namespace UE::XRCreative::Private
+{
+	static const FName PluginsCategoryName = TEXT("Plugins");
+};
+
+
+FName UXRCreativeSettings::GetCategoryName() const
+{
+	return UE::XRCreative::Private::PluginsCategoryName;
+}
+
 UXRCreativeSettings* UXRCreativeSettings::GetXRCreativeSettings()
 {
 	return GetMutableDefault<UXRCreativeSettings>();
 }
 
+FName UXRCreativeEditorSettings::GetCategoryName() const
+{
+	return UE::XRCreative::Private::PluginsCategoryName;
+}
+
 UXRCreativeEditorSettings* UXRCreativeEditorSettings::GetXRCreativeEditorSettings()
 {
 	return GetMutableDefault<UXRCreativeEditorSettings>();
diff --git a/Engine/Plugins/Experimental/XRCreativeFramework/Source/XRCreative/Public/XRCreativeSettings.h b/Engine/Plugins/Experimental/XRCreativeFramework/Source/XRCreative/Public/XRCreativeSettings.h
--- a/Engine/Plugins/Experimental/XRCreativeFramework/Source/XRCreative/Public/XRCreativeSettings.h
+++ b/Engine/Plugins/Experimental/XRCreativeFramework/Source/XRCreative/Public/XRCreativeSettings.h
@@ -2,6 +2,7 @@
 
 #pragma once
 
+#include "Engine/DeveloperSettings.h"
 #include "XRCreativeSettings.generated.h"
 
 
@@ -17,11 +18,15 @@ enum class EXRCreativeHandedness : uint8
  * Per project settings for XRCreative.
  */
 UCLASS(Config=XRCreativeSettings, DefaultConfig, DisplayName="XR Creative")
-class XRCREATIVE_API UXRCreativeSettings : public UObject
+class XRCREATIVE_API UXRCreativeSettings : public UDeveloperSettings
 {
 	GENERATED_BODY()
 	
 public:
+	//~ Begin UDeveloperSettings interface
+	virtual FName GetCategoryName() const override;
+	//~ End UDeveloperSettings interface
+
 	UFUNCTION(BlueprintPure, Category="XR Creative")
 	static UXRCreativeSettings* GetXRCreativeSettings();
 };
@@ -32,18 +37,21 @@ public:
  */
 UCLASS(Config=EditorPerProjectUserSettings, meta=(DisplayName="XR Creative Editor"))
 
-class XRCREATIVE_API UXRCreativeEditorSettings : public UObject
+class XRCREATIVE_API UXRCreativeEditorSettings : public UDeveloperSettings
 {
 	GENERATED_BODY()
 
 public:
+	//~ Begin UDeveloperSettings interface
+	virtual FName GetCategoryName() const override;
+	//~ End UDeveloperSettings interface
+
+	UFUNCTION(BlueprintPure, Category="XR Creative Editor")
+	static UXRCreativeEditorSettings* GetXRCreativeEditorSettings();
 
 	/** Manages Left/Right handedness user preferences.
 	 * Modifying this setting requires an editor restart to take effect.
 	 **/
 	UPROPERTY(Config, EditAnywhere, BlueprintReadWrite, Category="XR Creative", meta=(DisplayName="Handedness"))
 	EXRCreativeHandedness Handedness = EXRCreativeHandedness::Right;
-	
-	UFUNCTION(BlueprintPure, Category="XR Creative Editor")
-	static UXRCreativeEditorSettings* GetXRCreativeEditorSettings();
 };
diff --git a/Engine/Plugins/Experimental/XRCreativeFramework/Source/XRCreative/XRCreative.Build.cs b/Engine/Plugins/Experimental/XRCreativeFramework/Source/XRCreative/XRCreative.Build.cs
--- a/Engine/Plugins/Experimental/XRCreativeFramework/Source/XRCreative/XRCreative.Build.cs
+++ b/Engine/Plugins/Experimental/XRCreativeFramework/Source/XRCreative/XRCreative.Build.cs
@@ -30,7 +30,6 @@ public class XRCreative : ModuleRules
 				"LevelSequence",
 				"ModelViewViewModel",
 				"RenderCore", // TODO: Kill this? Only for FlushRenderingCommands
-				"Settings",
 				"Slate",
 				"SlateCore",
 				"TypedElementFramework",
