Virtual Scouting

Actor Placement Tool
- Polish pass on Content browser widget
- new icons for categories headings
- filters are now modal not aggregated
- added project setting to store swap to grab tool option
- fix layout in option submenu

Grab tool polish:
-added reset rotation functionality
-rearranged buttons

Gizmo polish:
- force materials to be hidden by menu if menu open when swapping modes from palette menu
- one frame delay spawning new gizmo

Slider Widget polish:
- fixed padding modifying with roundedbox handles - replaced with images that are all 48px square so padding stay consistent.
- applied same to color picker widget
-

Env lighting tool polish:
- Fix modal swapping of sun and moon buttons - clicking again on the same button no longer toggles selection state.

#rnx

#virtualized

diff --git a/Engine/Plugins/Experimental/VirtualScouting/Source/VirtualScouting/Public/VirtualScoutingSettings.h b/Engine/Plugins/Experimental/VirtualScouting/Source/VirtualScouting/Public/VirtualScoutingSettings.h
--- a/Engine/Plugins/Experimental/VirtualScouting/Source/VirtualScouting/Public/VirtualScoutingSettings.h
+++ b/Engine/Plugins/Experimental/VirtualScouting/Source/VirtualScouting/Public/VirtualScoutingSettings.h
@@ -28,6 +28,9 @@ public:
 	UPROPERTY(Config, EditAnywhere, BlueprintReadWrite, Category="Virtual Scouting", meta=(DisplayName="Viewfinder Use AutoExposure"))
 	bool bViewfinderUseExposure = false;
 
+	UPROPERTY(Config, EditAnywhere, BlueprintReadWrite, Category="Virtual Scouting", meta=(DisplayName="Swap to Grab on Spawn Actor"))
+	bool bSwapToGrabToolOnSpawnNewActor = true;
+
 	UPROPERTY(Config, EditAnywhere, BlueprintReadWrite, Category="Virtual Scouting", meta=(DisplayName="Viewfinder ExposureCompensation", ClampMin=-15, ClampMax=15))
 	float ViewfinderExposureCompensation = 1;
 	
