Virtual Scouting

-removed experimental flag and moved to Virtual Production folder

#virtualized

diff --git a/Engine/Plugins/Experimental/VirtualScouting/Config/BaseVirtualScouting.ini b/Engine/Plugins/Experimental/VirtualScouting/Config/BaseVirtualScouting.ini
deleted file mode 100644
diff --git a/Engine/Plugins/Experimental/VirtualScouting/Source/VirtualScouting/Private/Module.cpp b/Engine/Plugins/Experimental/VirtualScouting/Source/VirtualScouting/Private/Module.cpp
deleted file mode 100644
diff --git a/Engine/Plugins/Experimental/VirtualScouting/Source/VirtualScouting/Private/Module.h b/Engine/Plugins/Experimental/VirtualScouting/Source/VirtualScouting/Private/Module.h
deleted file mode 100644
diff --git a/Engine/Plugins/Experimental/VirtualScouting/Source/VirtualScouting/Private/VirtualScoutingLog.h b/Engine/Plugins/Experimental/VirtualScouting/Source/VirtualScouting/Private/VirtualScoutingLog.h
deleted file mode 100644
diff --git a/Engine/Plugins/Experimental/VirtualScouting/Source/VirtualScouting/Private/VirtualScoutingSettings.cpp b/Engine/Plugins/Experimental/VirtualScouting/Source/VirtualScouting/Private/VirtualScoutingSettings.cpp
deleted file mode 100644
diff --git a/Engine/Plugins/Experimental/VirtualScouting/Source/VirtualScouting/Public/VirtualScoutingSettings.h b/Engine/Plugins/Experimental/VirtualScouting/Source/VirtualScouting/Public/VirtualScoutingSettings.h
deleted file mode 100644
diff --git a/Engine/Plugins/Experimental/VirtualScouting/Source/VirtualScouting/VirtualScouting.Build.cs b/Engine/Plugins/Experimental/VirtualScouting/Source/VirtualScouting/VirtualScouting.Build.cs
deleted file mode 100644
diff --git a/Engine/Plugins/Experimental/VirtualScouting/Source/VirtualScoutingEditor/Private/VPScoutingMode.cpp b/Engine/Plugins/Experimental/VirtualScouting/Source/VirtualScoutingEditor/Private/VPScoutingMode.cpp
deleted file mode 100644
diff --git a/Engine/Plugins/Experimental/VirtualScouting/Source/VirtualScoutingEditor/Private/VPScoutingMode.h b/Engine/Plugins/Experimental/VirtualScouting/Source/VirtualScoutingEditor/Private/VPScoutingMode.h
deleted file mode 100644
diff --git a/Engine/Plugins/Experimental/VirtualScouting/Source/VirtualScoutingEditor/Private/VirtualScoutingBlueprintFunctionLibrary.cpp b/Engine/Plugins/Experimental/VirtualScouting/Source/VirtualScoutingEditor/Private/VirtualScoutingBlueprintFunctionLibrary.cpp
deleted file mode 100644
diff --git a/Engine/Plugins/Experimental/VirtualScouting/Source/VirtualScoutingEditor/Private/VirtualScoutingEditorModule.cpp b/Engine/Plugins/Experimental/VirtualScouting/Source/VirtualScoutingEditor/Private/VirtualScoutingEditorModule.cpp
deleted file mode 100644
diff --git a/Engine/Plugins/Experimental/VirtualScouting/Source/VirtualScoutingEditor/Private/VirtualScoutingEditorModule.h b/Engine/Plugins/Experimental/VirtualScouting/Source/VirtualScoutingEditor/Private/VirtualScoutingEditorModule.h
deleted file mode 100644
diff --git a/Engine/Plugins/Experimental/VirtualScouting/Source/VirtualScoutingEditor/Public/VirtualScoutingBlueprintFunctionLibrary.h b/Engine/Plugins/Experimental/VirtualScouting/Source/VirtualScoutingEditor/Public/VirtualScoutingBlueprintFunctionLibrary.h
deleted file mode 100644
diff --git a/Engine/Plugins/Experimental/VirtualScouting/Source/VirtualScoutingEditor/VirtualScoutingEditor.Build.cs b/Engine/Plugins/Experimental/VirtualScouting/Source/VirtualScoutingEditor/VirtualScoutingEditor.Build.cs
deleted file mode 100644
diff --git a/Engine/Plugins/Experimental/VirtualScouting/Source/VirtualScoutingOpenXR/Internal/VirtualScoutingOpenXR.h b/Engine/Plugins/Experimental/VirtualScouting/Source/VirtualScoutingOpenXR/Internal/VirtualScoutingOpenXR.h
deleted file mode 100644
diff --git a/Engine/Plugins/Experimental/VirtualScouting/Source/VirtualScoutingOpenXR/Internal/VirtualScoutingOpenXRModule.h b/Engine/Plugins/Experimental/VirtualScouting/Source/VirtualScoutingOpenXR/Internal/VirtualScoutingOpenXRModule.h
deleted file mode 100644
diff --git a/Engine/Plugins/Experimental/VirtualScouting/Source/VirtualScoutingOpenXR/Private/VirtualScoutingOpenXR.cpp b/Engine/Plugins/Experimental/VirtualScouting/Source/VirtualScoutingOpenXR/Private/VirtualScoutingOpenXR.cpp
deleted file mode 100644
diff --git a/Engine/Plugins/Experimental/VirtualScouting/Source/VirtualScoutingOpenXR/VirtualScoutingOpenXR.Build.cs b/Engine/Plugins/Experimental/VirtualScouting/Source/VirtualScoutingOpenXR/VirtualScoutingOpenXR.Build.cs
deleted file mode 100644
diff --git a/Engine/Plugins/Experimental/VirtualScouting/VirtualScouting.uplugin b/Engine/Plugins/Experimental/VirtualScouting/VirtualScouting.uplugin
deleted file mode 100644
diff --git a/Engine/Plugins/VirtualProduction/VirtualScouting/Config/BaseVirtualScouting.ini b/Engine/Plugins/VirtualProduction/VirtualScouting/Config/BaseVirtualScouting.ini
new file mode 100644
--- /dev/null
+++ b/Engine/Plugins/VirtualProduction/VirtualScouting/Config/BaseVirtualScouting.ini
@@ -0,0 +1,2 @@
+[/Script/VirtualScouting.VirtualScoutingSettings]
+bUseImperial=True
diff --git a/Engine/Plugins/VirtualProduction/VirtualScouting/Source/VirtualScouting/Private/Module.cpp b/Engine/Plugins/VirtualProduction/VirtualScouting/Source/VirtualScouting/Private/Module.cpp
new file mode 100644
--- /dev/null
+++ b/Engine/Plugins/VirtualProduction/VirtualScouting/Source/VirtualScouting/Private/Module.cpp
@@ -0,0 +1,26 @@
+﻿// Copyright Epic Games, Inc. All Rights Reserved.
+
+#include "Module.h"
+#include "VirtualScoutingLog.h"
+
+
+#define LOCTEXT_NAMESPACE "FVirtualScoutingModule"
+
+
+DEFINE_LOG_CATEGORY(LogVirtualScouting);
+
+
+void FVirtualScoutingModule::StartupModule()
+{
+}
+
+
+void FVirtualScoutingModule::ShutdownModule()
+{
+}
+
+
+IMPLEMENT_MODULE(FVirtualScoutingModule, VirtualScouting)
+
+
+#undef LOCTEXT_NAMESPACE
diff --git a/Engine/Plugins/VirtualProduction/VirtualScouting/Source/VirtualScouting/Private/Module.h b/Engine/Plugins/VirtualProduction/VirtualScouting/Source/VirtualScouting/Private/Module.h
new file mode 100644
--- /dev/null
+++ b/Engine/Plugins/VirtualProduction/VirtualScouting/Source/VirtualScouting/Private/Module.h
@@ -0,0 +1,13 @@
+﻿// Copyright Epic Games, Inc. All Rights Reserved.
+
+#pragma once
+
+#include "Modules/ModuleManager.h"
+
+
+class FVirtualScoutingModule : public IModuleInterface
+{
+public:
+	virtual void StartupModule() override;
+	virtual void ShutdownModule() override;
+};
diff --git a/Engine/Plugins/VirtualProduction/VirtualScouting/Source/VirtualScouting/Private/VirtualScoutingLog.h b/Engine/Plugins/VirtualProduction/VirtualScouting/Source/VirtualScouting/Private/VirtualScoutingLog.h
new file mode 100644
--- /dev/null
+++ b/Engine/Plugins/VirtualProduction/VirtualScouting/Source/VirtualScouting/Private/VirtualScoutingLog.h
@@ -0,0 +1,8 @@
+// Copyright Epic Games, Inc. All Rights Reserved.
+
+#pragma once
+
+#include "Logging/LogMacros.h"
+
+
+DECLARE_LOG_CATEGORY_EXTERN(LogVirtualScouting, Log, All);
diff --git a/Engine/Plugins/VirtualProduction/VirtualScouting/Source/VirtualScouting/Private/VirtualScoutingSettings.cpp b/Engine/Plugins/VirtualProduction/VirtualScouting/Source/VirtualScouting/Private/VirtualScoutingSettings.cpp
new file mode 100644
--- /dev/null
+++ b/Engine/Plugins/VirtualProduction/VirtualScouting/Source/VirtualScouting/Private/VirtualScoutingSettings.cpp
@@ -0,0 +1,30 @@
+// Copyright Epic Games, Inc. All Rights Reserved.
+
+#include "VirtualScoutingSettings.h"
+
+
+namespace UE::VirtualScouting::Private
+{
+	static const FName PluginsCategoryName = TEXT("Plugins");
+};
+
+
+FName UVirtualScoutingSettings::GetCategoryName() const
+{
+	return UE::VirtualScouting::Private::PluginsCategoryName;
+}
+
+UVirtualScoutingSettings* UVirtualScoutingSettings::GetVirtualScoutingSettings()
+{
+	return GetMutableDefault<UVirtualScoutingSettings>();
+}
+
+FName UVirtualScoutingEditorSettings::GetCategoryName() const
+{
+	return UE::VirtualScouting::Private::PluginsCategoryName;
+}
+
+UVirtualScoutingEditorSettings* UVirtualScoutingEditorSettings::GetVirtualScoutingEditorSettings()
+{
+	return GetMutableDefault<UVirtualScoutingEditorSettings>();
+}
diff --git a/Engine/Plugins/VirtualProduction/VirtualScouting/Source/VirtualScouting/Public/VirtualScoutingSettings.h b/Engine/Plugins/VirtualProduction/VirtualScouting/Source/VirtualScouting/Public/VirtualScoutingSettings.h
new file mode 100644
--- /dev/null
+++ b/Engine/Plugins/VirtualProduction/VirtualScouting/Source/VirtualScouting/Public/VirtualScoutingSettings.h
@@ -0,0 +1,80 @@
+// Copyright Epic Games, Inc. All Rights Reserved.
+
+#pragma once
+
+#include "Engine/DeveloperSettings.h"
+#include "VirtualScoutingSettings.generated.h"
+
+
+/**
+ * Per project settings for Virtual Scouting.
+ */
+UCLASS(Config=VirtualScoutingSettings, DefaultConfig, DisplayName="Virtual Scouting")
+class VIRTUALSCOUTING_API UVirtualScoutingSettings : public UDeveloperSettings
+{
+	GENERATED_BODY()
+	
+public:
+	//~ Begin UDeveloperSettings interface
+	virtual FName GetCategoryName() const override;
+	//~ End UDeveloperSettings interface
+
+	UFUNCTION(BlueprintPure, Category="Virtual Scouting", DisplayName="Virtual Scouting Settings")
+	static UVirtualScoutingSettings* GetVirtualScoutingSettings();
+
+	UPROPERTY(Config, EditAnywhere, BlueprintReadWrite, Category="Virtual Scouting", meta=(DisplayName="Show Measurements in Imperial Units"))
+	bool bUseImperial = false;
+
+	UPROPERTY(Config, EditAnywhere, BlueprintReadWrite, Category="Virtual Scouting", meta=(DisplayName="Viewfinder Use AutoExposure"))
+	bool bViewfinderUseExposure = false;
+
+	UPROPERTY(Config, EditAnywhere, BlueprintReadWrite, Category="Virtual Scouting", meta=(DisplayName="Swap to Grab on Spawn Actor"))
+	bool bSwapToGrabToolOnSpawnNewActor = true;
+
+	UPROPERTY(Config, EditAnywhere, BlueprintReadWrite, Category="Virtual Scouting", meta=(DisplayName="Viewfinder ExposureCompensation", ClampMin=-15, ClampMax=15))
+	float ViewfinderExposureCompensation = 1;
+	
+	UPROPERTY(Config, EditAnywhere, BlueprintReadWrite, Category="Virtual Scouting", meta=(DisplayName="Viewfinder Apertures"))
+	TArray<float> ViewfinderApertureArray = {1.2, 2.0, 2.8, 4.0, 5.6, 8.0, 11.0, 16.0, 22.0};
+
+	UPROPERTY(Config, EditAnywhere, BlueprintReadWrite, Category="Virtual Scouting", meta=(DisplayName="Viewfinder Monitor Masks"))
+	TArray<float> ViewfinderMaskArray = {1.33, 1.66, 1.78, 2.0, 2.35, 2.39 };
+
+	UPROPERTY(Config, EditAnywhere, BlueprintReadWrite, Category="Virtual Scouting", meta=(DisplayName="Sequence Tool Collection"))
+	FName SequenceToolCollection;
+
+	UPROPERTY(Config, EditAnywhere, BlueprintReadWrite, Category = "Virtual Scouting", meta = (DisplayName = "Content Browser Tool Collection"))
+	FName ContentBrowserToolCollection;
+};
+
+/**
+ * Per user settings for Virtual Scouting Editor.
+ */
+UCLASS(Config=EditorPerProjectUserSettings, DisplayName="Virtual Scouting Editor")
+class VIRTUALSCOUTING_API UVirtualScoutingEditorSettings : public UDeveloperSettings
+{
+	GENERATED_BODY()
+	
+public:
+	//~ Begin UDeveloperSettings interface
+	virtual FName GetCategoryName() const override;
+	//~ End UDeveloperSettings interface
+
+	UFUNCTION(BlueprintPure, Category="Virtual Scouting Editor")
+	static UVirtualScoutingEditorSettings* GetVirtualScoutingEditorSettings();
+
+	UPROPERTY(Config, EditAnywhere, BlueprintReadWrite, Category="Virtual Scouting User", meta=(DisplayName= "Flight Speed", ToolTip="Avatar flying speed. Default is 4.0", ClampMin=1.0f, ClampMax=10.0f));
+	float FlightSpeed = 4.0f;
+
+	UPROPERTY(Config, EditAnywhere, BlueprintReadWrite, Category="Virtual Scouting User", meta=(DisplayName= "Drag Speed", ToolTip="Speed of movement when you drag-move. Default is 0.7", ClampMin=0.1f, ClampMax=2.0f));
+	float DragSpeed = 0.7f;
+
+	UPROPERTY(Config, EditAnywhere, BlueprintReadWrite, Category="Virtual Scouting User", meta=(DisplayName= "Show Tooltips", ToolTip="Show Tooltips when in VR. These appear when motioncontroller is brought near to HMD"));
+	bool bEnableTooltips = true;
+
+	UPROPERTY(Config, EditAnywhere, BlueprintReadWrite, Category="Virtual Scouting User", meta=(DisplayName= "Use Smooth Rotation", ToolTip="True = Rotate smoothly. False = Flick Rotate. Default is Flick Rotate"));
+	bool bUseSmoothRotation = false;
+
+	UPROPERTY(Config, EditAnywhere, BlueprintReadWrite, Category="Virtual Scouting User", meta=(DisplayName= "Use Teleport Rotation", ToolTip="Use the forward axis roll from the motion controller to define and adjust teleport rotation"));
+	bool bUseTeleportRotation = false;
+};
diff --git a/Engine/Plugins/VirtualProduction/VirtualScouting/Source/VirtualScouting/VirtualScouting.Build.cs b/Engine/Plugins/VirtualProduction/VirtualScouting/Source/VirtualScouting/VirtualScouting.Build.cs
new file mode 100644
--- /dev/null
+++ b/Engine/Plugins/VirtualProduction/VirtualScouting/Source/VirtualScouting/VirtualScouting.Build.cs
@@ -0,0 +1,53 @@
+// Copyright Epic Games, Inc. All Rights Reserved.
+
+using UnrealBuildTool;
+
+public class VirtualScouting : ModuleRules
+{
+	public VirtualScouting(ReadOnlyTargetRules Target) : base(Target)
+	{
+		PCHUsage = ModuleRules.PCHUsageMode.UseExplicitOrSharedPCHs;
+
+		PublicDependencyModuleNames.AddRange(
+			new string[]
+			{
+				"CommonUI",
+				"Core",
+				"XRCreative",
+			}
+		);
+
+		PrivateDependencyModuleNames.AddRange(
+			new string[]
+			{
+				"CoreUObject",
+				"DeveloperSettings",
+				"Engine",
+				"EnhancedInput",
+				"HeadMountedDisplay",
+				"InputCore",
+				"InteractiveToolsFramework",
+				"LevelSequence",
+				"RenderCore", // TODO: Kill this? Only for FlushRenderingCommands
+				"Slate",
+				"SlateCore",
+				"TypedElementFramework",
+				"TypedElementRuntime",
+				"UMG",
+			}
+		);
+
+		if (Target.bBuildEditor)
+		{
+			PrivateDependencyModuleNames.AddRange(new string[] {
+				"ConcertSyncClient",
+				"ConcertSyncCore",
+				"InputEditor",
+				"LevelEditor",
+				"MultiUserClient",
+				"UnrealEd",
+				"VREditor",
+			});
+		}
+	}
+}
diff --git a/Engine/Plugins/VirtualProduction/VirtualScouting/Source/VirtualScoutingEditor/Private/VPScoutingMode.cpp b/Engine/Plugins/VirtualProduction/VirtualScouting/Source/VirtualScoutingEditor/Private/VPScoutingMode.cpp
new file mode 100644
--- /dev/null
+++ b/Engine/Plugins/VirtualProduction/VirtualScouting/Source/VirtualScoutingEditor/Private/VPScoutingMode.cpp
@@ -0,0 +1,136 @@
+// Copyright Epic Games, Inc. All Rights Reserved.
+
+#include "VPScoutingMode.h"
+#include "VirtualScoutingEditorModule.h"
+#include "VirtualScoutingOpenXR.h"
+#include "VirtualScoutingOpenXRModule.h"
+
+#if WITH_EDITOR
+#include "Editor.h"
+#endif
+
+#include "Framework/Notifications/NotificationManager.h"
+#include "IOpenXRHMDModule.h"
+#include "IVREditorModule.h"
+#include "IXRTrackingSystem.h"
+#include "Widgets/Notifications/SNotificationList.h"
+
+
+#define LOCTEXT_NAMESPACE "VirtualScouting"
+
+
+static const FName OpenXRSystemName = FName(TEXT("OpenXR"));
+
+
+UVPScoutingMode::UVPScoutingMode()
+{
+	InteractorClass = FSoftObjectPath("/VirtualProductionUtilities/VR/VirtualScoutingInteractor.VirtualScoutingInteractor_C");
+	TeleporterClass = FSoftObjectPath("/VirtualProductionUtilities/VR/VirtualScoutingTeleporter.VirtualScoutingTeleporter_C");
+}
+
+
+bool UVPScoutingMode::NeedsSyntheticDpad()
+{
+	if (!GEngine->XRSystem.IsValid() || GEngine->XRSystem->GetSystemName() != OpenXRSystemName)
+	{
+		return Super::NeedsSyntheticDpad();
+	}
+
+	return !IOpenXRHMDModule::Get().IsExtensionEnabled("XR_EXT_dpad_binding");
+}
+
+
+void UVPScoutingMode::Enter()
+{
+	if (!GEngine->XRSystem.IsValid() || GEngine->XRSystem->GetSystemName() != OpenXRSystemName)
+	{
+		Super::Enter();
+		return;
+	}
+
+	TSharedPtr<FVirtualScoutingOpenXRExtension> XrExt = FVirtualScoutingOpenXRModule::Get().GetOpenXRExt();
+	if (!XrExt)
+	{
+		UE_LOG(LogVirtualScoutingEditor, Error, TEXT("OpenXR extension plugin invalid"));
+		IVREditorModule::Get().EnableVREditor(false);
+		return;
+	}
+
+	if (!ValidateSettings())
+	{
+		IVREditorModule::Get().EnableVREditor(false);
+		return;
+	}
+
+#if WITH_EDITOR
+	// This causes FOpenXRInput to rebuild and reattach actions.
+	FEditorDelegates::OnActionAxisMappingsChanged.Broadcast();
+#endif
+
+	// Split the mode entry into two phases. This is necessary because we have to poll OpenXR for
+	// the active interaction profile and translate it into a legacy plugin name, but OpenXR may
+	// not return the correct interaction profile for several frames after the OpenXR session
+	// (stereo rendering) has started, and we need to defer creation of the interactors, etc.
+	BeginEntry();
+
+	XrExt->GetHmdDeviceTypeFuture() = XrExt->GetHmdDeviceTypeFuture().Next(
+		[WeakThis = TWeakObjectPtr<UVPScoutingMode>(this)]
+		(FName DeviceType)
+		{
+			UVPScoutingMode* This = WeakThis.Get();
+			if (!This || !This->bIsFullyInitialized)
+			{
+				UE_LOG(LogVirtualScoutingEditor, Warning, TEXT("Stale UVPScoutingMode; ignoring HmdDeviceType future"));
+				return DeviceType;
+			}
+
+			if (DeviceType == NAME_None)
+			{
+				UE_LOG(LogVirtualScoutingEditor, Error, TEXT("Unable to map legacy HMD device type"));
+				return DeviceType;
+			}
+
+			This->SetHMDDeviceTypeOverride(DeviceType);
+			This->SetupSubsystems();
+			This->FinishEntry();
+
+			return DeviceType;
+		});
+}
+
+
+bool UVPScoutingMode::ValidateSettings()
+{
+	bool bSettingsValid = true;
+
+	IConsoleManager& ConsoleMgr = IConsoleManager::Get();
+	if (IConsoleVariable* PropagateAlpha = ConsoleMgr.FindConsoleVariable(TEXT("r.PostProcessing.PropagateAlpha")))
+	{
+		if (PropagateAlpha->GetBool())
+		{
+			InvalidSettingNotification(LOCTEXT("InvalidCvarPropagateAlpha", "r.PostProcessing.PropagateAlpha must be disabled (and requires an engine restart)"));
+			bSettingsValid = false;
+		}
+	}
+
+	return bSettingsValid;
+}
+
+
+void UVPScoutingMode::InvalidSettingNotification(const FText& ErrorDetails)
+{
+	FText NotificationHeading = LOCTEXT("InvalidSettingNotificationTitle", "Unable to initialize Virtual Scouting");
+
+	UE_LOG(LogVirtualScoutingEditor, Error, TEXT("%s: %s"), *NotificationHeading.ToString(), *ErrorDetails.ToString());
+
+	FNotificationInfo Info(NotificationHeading);
+	Info.SubText = ErrorDetails;
+	TSharedPtr<SNotificationItem> Notification = FSlateNotificationManager::Get().AddNotification(Info);
+	if (Notification)
+	{
+		Notification->SetCompletionState(SNotificationItem::CS_Fail);
+	}
+}
+
+
+#undef LOCTEXT_NAMESPACE
diff --git a/Engine/Plugins/VirtualProduction/VirtualScouting/Source/VirtualScoutingEditor/Private/VPScoutingMode.h b/Engine/Plugins/VirtualProduction/VirtualScouting/Source/VirtualScoutingEditor/Private/VPScoutingMode.h
new file mode 100644
--- /dev/null
+++ b/Engine/Plugins/VirtualProduction/VirtualScouting/Source/VirtualScoutingEditor/Private/VPScoutingMode.h
@@ -0,0 +1,27 @@
+// Copyright Epic Games, Inc. All Rights Reserved.
+
+#pragma once
+
+#include "VREditorMode.h"
+#include "VPScoutingMode.generated.h"
+
+
+UCLASS(meta=(DisplayName="Legacy Virtual Scouting"))
+class UVPScoutingMode : public UVREditorMode
+{
+	GENERATED_BODY()
+
+public:
+	UVPScoutingMode();
+
+	//~ Begin UVREditorMode interface
+	virtual bool NeedsSyntheticDpad() override;
+
+	virtual void Enter() override;
+	//~ End UVREditorMode interface
+
+protected:
+	bool ValidateSettings();
+
+	void InvalidSettingNotification(const FText& ErrorDetails);
+};
diff --git a/Engine/Plugins/VirtualProduction/VirtualScouting/Source/VirtualScoutingEditor/Private/VirtualScoutingBlueprintFunctionLibrary.cpp b/Engine/Plugins/VirtualProduction/VirtualScouting/Source/VirtualScoutingEditor/Private/VirtualScoutingBlueprintFunctionLibrary.cpp
new file mode 100644
--- /dev/null
+++ b/Engine/Plugins/VirtualProduction/VirtualScouting/Source/VirtualScoutingEditor/Private/VirtualScoutingBlueprintFunctionLibrary.cpp
@@ -0,0 +1,16 @@
+// Copyright Epic Games, Inc. All Rights Reserved.
+
+
+#include "VirtualScoutingBlueprintFunctionLibrary.h"
+
+#include "Editor.h"
+#include "ScopedTransaction.h"
+
+bool UVirtualScoutingBlueprintFunctionLibrary::CheckIsWithEditor()
+{
+#if WITH_EDITOR
+	return true;
+#else
+	return false;
+#endif
+}
diff --git a/Engine/Plugins/VirtualProduction/VirtualScouting/Source/VirtualScoutingEditor/Private/VirtualScoutingEditorModule.cpp b/Engine/Plugins/VirtualProduction/VirtualScouting/Source/VirtualScoutingEditor/Private/VirtualScoutingEditorModule.cpp
new file mode 100644
--- /dev/null
+++ b/Engine/Plugins/VirtualProduction/VirtualScouting/Source/VirtualScoutingEditor/Private/VirtualScoutingEditorModule.cpp
@@ -0,0 +1,31 @@
+// Copyright Epic Games, Inc. All Rights Reserved.
+
+#include "VirtualScoutingEditorModule.h"
+#include "Modules/ModuleManager.h"
+#include "VRModeSettings.h"
+
+
+DEFINE_LOG_CATEGORY(LogVirtualScoutingEditor);
+
+
+class FVirtualScoutingEditorModule : public IModuleInterface
+{
+private:
+	virtual void StartupModule() override
+	{
+		// When the user has no previous selection upon enabling this plugin
+		// for the first time, ensure we default to our preferred mode.
+		UVRModeSettings* Settings = GetMutableDefault<UVRModeSettings>();
+		if (Settings->ModeClass.IsNull())
+		{
+			Settings->ModeClass = FSoftObjectPath("/VirtualScouting/Core/Scouting_Default.Scouting_Default_C");
+		}
+	}
+
+	virtual void ShutdownModule() override
+	{
+	}
+};
+
+
+IMPLEMENT_MODULE(FVirtualScoutingEditorModule, VirtualScoutingEditor);
diff --git a/Engine/Plugins/VirtualProduction/VirtualScouting/Source/VirtualScoutingEditor/Private/VirtualScoutingEditorModule.h b/Engine/Plugins/VirtualProduction/VirtualScouting/Source/VirtualScoutingEditor/Private/VirtualScoutingEditorModule.h
new file mode 100644
--- /dev/null
+++ b/Engine/Plugins/VirtualProduction/VirtualScouting/Source/VirtualScoutingEditor/Private/VirtualScoutingEditorModule.h
@@ -0,0 +1,8 @@
+// Copyright Epic Games, Inc. All Rights Reserved.
+
+#pragma once
+
+#include "Logging/LogMacros.h"
+
+
+DECLARE_LOG_CATEGORY_EXTERN(LogVirtualScoutingEditor, Log, All);
diff --git a/Engine/Plugins/VirtualProduction/VirtualScouting/Source/VirtualScoutingEditor/Public/VirtualScoutingBlueprintFunctionLibrary.h b/Engine/Plugins/VirtualProduction/VirtualScouting/Source/VirtualScoutingEditor/Public/VirtualScoutingBlueprintFunctionLibrary.h
new file mode 100644
--- /dev/null
+++ b/Engine/Plugins/VirtualProduction/VirtualScouting/Source/VirtualScoutingEditor/Public/VirtualScoutingBlueprintFunctionLibrary.h
@@ -0,0 +1,20 @@
+// Copyright Epic Games, Inc. All Rights Reserved.
+
+#pragma once
+
+#include "CoreMinimal.h"
+#include "Kismet/BlueprintFunctionLibrary.h"
+#include "VirtualScoutingBlueprintFunctionLibrary.generated.h"
+
+/**
+ * 
+ */
+UCLASS()
+class VIRTUALSCOUTINGEDITOR_API UVirtualScoutingBlueprintFunctionLibrary : public UBlueprintFunctionLibrary
+{
+	GENERATED_BODY()
+
+	UFUNCTION(BlueprintCallable, Category="VirtualScouting")
+	static bool CheckIsWithEditor();
+
+};
diff --git a/Engine/Plugins/VirtualProduction/VirtualScouting/Source/VirtualScoutingEditor/VirtualScoutingEditor.Build.cs b/Engine/Plugins/VirtualProduction/VirtualScouting/Source/VirtualScoutingEditor/VirtualScoutingEditor.Build.cs
new file mode 100644
--- /dev/null
+++ b/Engine/Plugins/VirtualProduction/VirtualScouting/Source/VirtualScoutingEditor/VirtualScoutingEditor.Build.cs
@@ -0,0 +1,41 @@
+// Copyright Epic Games, Inc. All Rights Reserved.
+
+using UnrealBuildTool;
+
+public class VirtualScoutingEditor : ModuleRules
+{
+	public VirtualScoutingEditor(ReadOnlyTargetRules Target) : base(Target)
+	{
+		PCHUsage = ModuleRules.PCHUsageMode.UseExplicitOrSharedPCHs;
+
+		PublicDependencyModuleNames.AddRange(new string[]
+			{
+				"Core",
+				"CoreUObject",
+				"Engine",
+			}
+		);
+
+		PrivateDependencyModuleNames.AddRange(new string[]
+			{
+				"OpenXRHMD",
+				"OpenXRInput",
+				"UnrealEd",
+				"VREditor",
+				"VirtualScoutingOpenXR",
+				"EnhancedInput",
+				"HeadMountedDisplay",
+				"InputCore",
+				"InputEditor",
+				"InteractiveToolsFramework",
+				"LevelEditor",
+				"RenderCore",
+				"Slate",
+				"SlateCore",
+				"UMG",
+				"VirtualScouting",
+				"XRCreative",
+			}
+		);
+	}
+}
diff --git a/Engine/Plugins/VirtualProduction/VirtualScouting/Source/VirtualScoutingOpenXR/Internal/VirtualScoutingOpenXR.h b/Engine/Plugins/VirtualProduction/VirtualScouting/Source/VirtualScoutingOpenXR/Internal/VirtualScoutingOpenXR.h
new file mode 100644
--- /dev/null
+++ b/Engine/Plugins/VirtualProduction/VirtualScouting/Source/VirtualScoutingOpenXR/Internal/VirtualScoutingOpenXR.h
@@ -0,0 +1,66 @@
+// Copyright Epic Games, Inc. All Rights Reserved.
+
+#pragma once
+
+#if WITH_EDITOR
+
+#include "IOpenXRExtensionPlugin.h"
+#include "Async/Future.h"
+#include <openxr/openxr_platform_defines.h>
+
+template <typename OptionalType> struct TOptional;
+
+
+class VIRTUALSCOUTINGOPENXR_API FVirtualScoutingOpenXRExtension
+	: public IOpenXRExtensionPlugin
+	, public TSharedFromThis<FVirtualScoutingOpenXRExtension>
+{
+public:
+	FVirtualScoutingOpenXRExtension();
+	virtual ~FVirtualScoutingOpenXRExtension();
+
+	//~ Begin IOpenXRExtensionPlugin interface
+	virtual FString GetDisplayName() override
+	{
+		return FString(TEXT("VirtualScouting"));
+	}
+
+	virtual bool GetOptionalExtensions(TArray<const ANSICHAR*>& OutExtensions) override;
+	virtual void OnEvent(XrSession InSession, const XrEventDataBaseHeader* InHeader) override;
+	virtual void PostCreateInstance(XrInstance InInstance) override;
+	virtual void PostCreateSession(XrSession InSession) override;
+	virtual void PostSyncActions(XrSession InSession) override;
+	//~ End IOpenXRExtensionPlugin interface
+
+	TFuture<FName>& GetHmdDeviceTypeFuture() { return DeviceTypeFuture; }
+
+private:
+	void OnVREditingModeEnter();
+	void OnVREditingModeExit();
+
+	TOptional<FName> TryGetHmdDeviceType();
+	void TryFulfillDeviceTypePromise();
+
+private:
+	FDelegateHandle InitCompleteDelegate;
+
+	XrDebugUtilsMessengerEXT Messenger = XR_NULL_HANDLE;
+
+	XrInstance Instance = XR_NULL_HANDLE;
+	XrSession Session = XR_NULL_HANDLE;
+	XrActionSet ActionSet = XR_NULL_HANDLE;
+
+	bool bIsVrEditingModeActive = false;
+
+	TPromise<FName> DeviceTypePromise;
+	TFuture<FName> DeviceTypeFuture;
+
+private:
+	static XrBool32 XRAPI_CALL XrDebugUtilsMessengerCallback_Trampoline(
+		XrDebugUtilsMessageSeverityFlagsEXT InMessageSeverity,
+		XrDebugUtilsMessageTypeFlagsEXT InMessageTypes,
+		const XrDebugUtilsMessengerCallbackDataEXT* InCallbackData,
+		void* InUserData);
+};
+
+#endif // #if WITH_EDITOR
diff --git a/Engine/Plugins/VirtualProduction/VirtualScouting/Source/VirtualScoutingOpenXR/Internal/VirtualScoutingOpenXRModule.h b/Engine/Plugins/VirtualProduction/VirtualScouting/Source/VirtualScoutingOpenXR/Internal/VirtualScoutingOpenXRModule.h
new file mode 100644
--- /dev/null
+++ b/Engine/Plugins/VirtualProduction/VirtualScouting/Source/VirtualScoutingOpenXR/Internal/VirtualScoutingOpenXRModule.h
@@ -0,0 +1,33 @@
+// Copyright Epic Games, Inc. All Rights Reserved.
+
+#pragma once
+
+#include "Modules/ModuleManager.h"
+
+
+class FVirtualScoutingOpenXRExtension;
+
+
+class VIRTUALSCOUTINGOPENXR_API FVirtualScoutingOpenXRModule : public IModuleInterface
+{
+public:
+	/**
+	 * Singleton-like access to this module's interface.  This is just for convenience!
+	 * Beware of calling this during the shutdown phase, though.  Your module might have been unloaded already.
+	 *
+	 * @return Returns singleton instance, loading the module on demand if needed
+	 */
+	static FVirtualScoutingOpenXRModule& Get()
+	{
+		static const FName ModuleName = "VirtualScoutingOpenXR";
+		return FModuleManager::LoadModuleChecked<FVirtualScoutingOpenXRModule>(ModuleName);
+	};
+
+	virtual void StartupModule() override;
+	virtual void ShutdownModule() override;
+
+	const TSharedPtr<FVirtualScoutingOpenXRExtension>& GetOpenXRExt() { return OpenXRExt; }
+
+private:
+	TSharedPtr<FVirtualScoutingOpenXRExtension> OpenXRExt;
+};
diff --git a/Engine/Plugins/VirtualProduction/VirtualScouting/Source/VirtualScoutingOpenXR/Private/VirtualScoutingOpenXR.cpp b/Engine/Plugins/VirtualProduction/VirtualScouting/Source/VirtualScoutingOpenXR/Private/VirtualScoutingOpenXR.cpp
new file mode 100644
--- /dev/null
+++ b/Engine/Plugins/VirtualProduction/VirtualScouting/Source/VirtualScoutingOpenXR/Private/VirtualScoutingOpenXR.cpp
@@ -0,0 +1,290 @@
+// Copyright Epic Games, Inc. All Rights Reserved.
+
+#include "VirtualScoutingOpenXR.h"
+#include "Misc/CoreDelegates.h"
+#include "VirtualScoutingOpenXRModule.h"
+
+#include "OpenXRCore.h"
+
+#if WITH_EDITOR
+
+#include "IVREditorModule.h"
+
+
+#define LOCTEXT_NAMESPACE "VirtualScouting"
+
+
+DECLARE_LOG_CATEGORY_EXTERN(LogVirtualScoutingOpenXR, Log, All);
+DECLARE_LOG_CATEGORY_EXTERN(LogVirtualScoutingOpenXRDebug, VeryVerbose, All);
+DEFINE_LOG_CATEGORY(LogVirtualScoutingOpenXR);
+DEFINE_LOG_CATEGORY(LogVirtualScoutingOpenXRDebug);
+
+
+static TAutoConsoleVariable<int32> CVarOpenXRDebugLogging(
+	TEXT("VirtualScouting.OpenXRDebugLogging"),
+	0,
+	TEXT("If true, register an Unreal log sink via XR_EXT_debug_utils.\n"),
+	ECVF_Default);
+
+
+
+void FVirtualScoutingOpenXRModule::StartupModule()
+{
+	OpenXRExt = MakeShared<FVirtualScoutingOpenXRExtension>();
+}
+
+void FVirtualScoutingOpenXRModule::ShutdownModule()
+{
+	OpenXRExt.Reset();
+}
+
+
+IMPLEMENT_MODULE(FVirtualScoutingOpenXRModule, VirtualScoutingOpenXR);
+
+
+FVirtualScoutingOpenXRExtension::FVirtualScoutingOpenXRExtension()
+{
+	RegisterOpenXRExtensionModularFeature();
+
+	InitCompleteDelegate = FCoreDelegates::OnFEngineLoopInitComplete.AddLambda(
+		[this]()
+		{
+			IVREditorModule& VrEditor = IVREditorModule::Get();
+			VrEditor.OnVREditingModeEnter().AddRaw(this, &FVirtualScoutingOpenXRExtension::OnVREditingModeEnter);
+			VrEditor.OnVREditingModeExit().AddRaw(this, &FVirtualScoutingOpenXRExtension::OnVREditingModeExit);
+
+			// Must happen last; this implicitly deallocates this lambda's captures.
+			FCoreDelegates::OnFEngineLoopInitComplete.Remove(InitCompleteDelegate);
+		}
+	);
+}
+
+
+FVirtualScoutingOpenXRExtension::~FVirtualScoutingOpenXRExtension()
+{
+	if (!DeviceTypeFuture.IsReady())
+	{
+		DeviceTypePromise.EmplaceValue(NAME_None);
+	}
+
+#if 0 // FIXME?: Too late to use Instance here, and don't see any suitable extension interface methods.
+	// Might also be OK not to explicitly clean this up.
+	if (Messenger != XR_NULL_HANDLE)
+	{
+		PFN_xrDestroyDebugUtilsMessengerEXT PfnXrDestroyDebugUtilsMessengerEXT;
+		if (XR_ENSURE(xrGetInstanceProcAddr(Instance, "xrDestroyDebugUtilsMessengerEXT",
+			reinterpret_cast<PFN_xrVoidFunction*>(&PfnXrDestroyDebugUtilsMessengerEXT))))
+		{
+			XR_ENSURE(PfnXrDestroyDebugUtilsMessengerEXT(Messenger));
+			Messenger = XR_NULL_HANDLE;
+		}
+	}
+#endif
+
+	UnregisterOpenXRExtensionModularFeature();
+
+	if (IVREditorModule::IsAvailable())
+	{
+		IVREditorModule::Get().OnVREditingModeEnter().RemoveAll(this);
+		IVREditorModule::Get().OnVREditingModeExit().RemoveAll(this);
+	}
+}
+
+
+bool FVirtualScoutingOpenXRExtension::GetOptionalExtensions(TArray<const ANSICHAR*>& OutExtensions)
+{
+	if (CVarOpenXRDebugLogging.GetValueOnAnyThread() != 0)
+	{
+		OutExtensions.Add(XR_EXT_DEBUG_UTILS_EXTENSION_NAME);
+	}
+
+	return true;
+}
+
+
+void FVirtualScoutingOpenXRExtension::OnEvent(XrSession InSession, const XrEventDataBaseHeader* InHeader)
+{
+	switch (InHeader->type)
+	{
+		case XR_TYPE_EVENT_DATA_INTERACTION_PROFILE_CHANGED:
+			// TODO: Correctly handle interaction profile changing mid-session
+			ensure(InSession == Session);
+			TryFulfillDeviceTypePromise();
+			break;
+	}
+}
+
+
+void FVirtualScoutingOpenXRExtension::PostCreateInstance(XrInstance InInstance)
+{
+	Instance = InInstance;
+
+	if (CVarOpenXRDebugLogging.GetValueOnAnyThread() != 0)
+	{
+		PFN_xrCreateDebugUtilsMessengerEXT PfnXrCreateDebugUtilsMessengerEXT;
+		if (XR_ENSURE(xrGetInstanceProcAddr(Instance, "xrCreateDebugUtilsMessengerEXT",
+			reinterpret_cast<PFN_xrVoidFunction*>(&PfnXrCreateDebugUtilsMessengerEXT))))
+		{
+			XrDebugUtilsMessengerCreateInfoEXT DebugMessengerCreateInfo;
+			DebugMessengerCreateInfo.type = XR_TYPE_DEBUG_UTILS_MESSENGER_CREATE_INFO_EXT;
+			DebugMessengerCreateInfo.next = nullptr;
+			DebugMessengerCreateInfo.messageSeverities =
+				XR_DEBUG_UTILS_MESSAGE_SEVERITY_ERROR_BIT_EXT | XR_DEBUG_UTILS_MESSAGE_SEVERITY_WARNING_BIT_EXT |
+				XR_DEBUG_UTILS_MESSAGE_SEVERITY_INFO_BIT_EXT | XR_DEBUG_UTILS_MESSAGE_SEVERITY_VERBOSE_BIT_EXT;
+			DebugMessengerCreateInfo.messageTypes =
+				XR_DEBUG_UTILS_MESSAGE_TYPE_GENERAL_BIT_EXT | XR_DEBUG_UTILS_MESSAGE_TYPE_VALIDATION_BIT_EXT |
+				XR_DEBUG_UTILS_MESSAGE_TYPE_PERFORMANCE_BIT_EXT | XR_DEBUG_UTILS_MESSAGE_TYPE_CONFORMANCE_BIT_EXT;
+			DebugMessengerCreateInfo.userCallback = &FVirtualScoutingOpenXRExtension::XrDebugUtilsMessengerCallback_Trampoline;
+			DebugMessengerCreateInfo.userData = this;
+
+			if (XR_ENSURE(PfnXrCreateDebugUtilsMessengerEXT(Instance, &DebugMessengerCreateInfo, &Messenger)))
+			{
+				UE_LOG(LogVirtualScoutingOpenXRDebug, Log, TEXT("XR_EXT_debug_utils messenger ACTIVE"));
+				return;
+			}
+		}
+	}
+
+	UE_LOG(LogVirtualScoutingOpenXRDebug, Log, TEXT("XR_EXT_debug_utils messenger DISABLED"));
+}
+
+
+void FVirtualScoutingOpenXRExtension::PostCreateSession(XrSession InSession)
+{
+	Session = InSession;
+
+	DeviceTypePromise = {};
+	DeviceTypeFuture = DeviceTypePromise.GetFuture();
+}
+
+
+void FVirtualScoutingOpenXRExtension::PostSyncActions(XrSession InSession)
+{
+	TryFulfillDeviceTypePromise();
+}
+
+
+void FVirtualScoutingOpenXRExtension::TryFulfillDeviceTypePromise()
+{
+	if (DeviceTypeFuture.IsValid() && !DeviceTypeFuture.IsReady())
+	{
+		TOptional<FName> MaybeType = TryGetHmdDeviceType();
+		if (MaybeType)
+		{
+			DeviceTypePromise.EmplaceValue(MaybeType.GetValue());
+			return;
+		}
+	}
+}
+
+
+TOptional<FName> FVirtualScoutingOpenXRExtension::TryGetHmdDeviceType()
+{
+	if (Instance == XR_NULL_HANDLE || Session == XR_NULL_HANDLE)
+	{
+		return FName();
+	}
+
+	XrPath LeftHand = XR_NULL_PATH, RightHand = XR_NULL_PATH;
+	check(XR_SUCCEEDED(xrStringToPath(Instance, "/user/hand/left", &LeftHand)));
+	check(XR_SUCCEEDED(xrStringToPath(Instance, "/user/hand/right", &RightHand)));
+
+	XrInteractionProfileState LeftHandState{ XR_TYPE_INTERACTION_PROFILE_STATE },
+	                          RightHandState{ XR_TYPE_INTERACTION_PROFILE_STATE };
+
+	XR_ENSURE(xrGetCurrentInteractionProfile(Session, LeftHand, &LeftHandState));
+	XR_ENSURE(xrGetCurrentInteractionProfile(Session, RightHand, &RightHandState));
+
+	// TODO: Correctly handle mixed interaction profiles?
+	XrPath InteractionProfile = LeftHandState.interactionProfile != XR_NULL_PATH
+		? LeftHandState.interactionProfile
+		: RightHandState.interactionProfile;
+
+	if (InteractionProfile == XR_NULL_PATH)
+	{
+		// Returning an unset optional indicates we don't know _yet_, and the caller can poll again.
+		return TOptional<FName>();
+	}
+
+	uint32 PathLen;
+	char PathBuf[XR_MAX_PATH_LENGTH] = { 0 };
+	if (!XR_ENSURE(xrPathToString(Instance, InteractionProfile, XR_MAX_PATH_LENGTH, &PathLen, PathBuf)))
+	{
+		return FName();
+	}
+
+	const FAnsiStringView PathView(PathBuf, PathLen - 1);
+	if (PathView.StartsWith(ANSITEXTVIEW("/interaction_profiles/oculus/")))
+	{
+		return FName("OculusHMD");
+	}
+	else if (PathView.StartsWith(ANSITEXTVIEW("/interaction_profiles/htc/"))
+		|| PathView.StartsWith(ANSITEXTVIEW("/interaction_profiles/valve/")))
+	{
+		return FName("SteamVR");
+	}
+
+	UE_LOG(LogVirtualScoutingOpenXR, Error, TEXT("Unsupported interaction profile: '%S'"), PathBuf);
+	return FName();
+}
+
+
+void FVirtualScoutingOpenXRExtension::OnVREditingModeEnter()
+{
+	bIsVrEditingModeActive = true;
+}
+
+
+void FVirtualScoutingOpenXRExtension::OnVREditingModeExit()
+{
+	bIsVrEditingModeActive = false;
+}
+
+
+//static
+XrBool32 XRAPI_CALL FVirtualScoutingOpenXRExtension::XrDebugUtilsMessengerCallback_Trampoline(
+	XrDebugUtilsMessageSeverityFlagsEXT InMessageSeverity,
+	XrDebugUtilsMessageTypeFlagsEXT InMessageTypes,
+	const XrDebugUtilsMessengerCallbackDataEXT* InCallbackData,
+	void* InUserData)
+{
+	// TODO?: InUserData contains the FVirtualScoutingOpenXRExtension*, could forward to an instance method.
+	// However, doing the work directly in here seems fine for now.
+
+	ELogVerbosity::Type Verbosity;
+	switch (InMessageSeverity)
+	{
+		case XR_DEBUG_UTILS_MESSAGE_SEVERITY_VERBOSE_BIT_EXT: Verbosity = ELogVerbosity::Verbose; break;
+		case XR_DEBUG_UTILS_MESSAGE_SEVERITY_INFO_BIT_EXT:    Verbosity = ELogVerbosity::Display; break;
+		case XR_DEBUG_UTILS_MESSAGE_SEVERITY_WARNING_BIT_EXT: Verbosity = ELogVerbosity::Warning; break;
+		case XR_DEBUG_UTILS_MESSAGE_SEVERITY_ERROR_BIT_EXT:   Verbosity = ELogVerbosity::Error;   break;
+		default:
+			ensureMsgf(false, TEXT("Unhandled XrDebugUtilsMessageSeverityFlagsEXT: %X"), InMessageSeverity);
+			Verbosity = ELogVerbosity::Error;
+			break;
+	}
+
+	// Results in "____" if no bits are set, or "GVPC" if all bits are set, or some combination.
+	FString Types = TEXT("####");
+	Types[0] = (InMessageTypes & XR_DEBUG_UTILS_MESSAGE_TYPE_GENERAL_BIT_EXT)     ? TEXT('G') : TEXT('_');
+	Types[1] = (InMessageTypes & XR_DEBUG_UTILS_MESSAGE_TYPE_VALIDATION_BIT_EXT)  ? TEXT('V') : TEXT('_');
+	Types[2] = (InMessageTypes & XR_DEBUG_UTILS_MESSAGE_TYPE_PERFORMANCE_BIT_EXT) ? TEXT('P') : TEXT('_');
+	Types[3] = (InMessageTypes & XR_DEBUG_UTILS_MESSAGE_TYPE_CONFORMANCE_BIT_EXT) ? TEXT('C') : TEXT('_');
+
+	FMsg::Logf(__FILE__, __LINE__, LogVirtualScoutingOpenXRDebug.GetCategoryName(), Verbosity,
+		TEXT("[%s]: %S(): %S"), *Types, InCallbackData->functionName, InCallbackData->message);
+
+	// "A value of XR_TRUE indicates that the application wants to abort this call. [...]
+	// Applications should always return XR_FALSE so that they see the same behavior with
+	// and without validation layers enabled."
+	return XR_FALSE;
+}
+
+
+#undef LOCTEXT_NAMESPACE
+
+#else // not WITH_EDITOR
+
+IMPLEMENT_MODULE(FDefaultModuleImpl, VirtualScoutingOpenXR);
+
+#endif // #if WITH_EDITOR .. #else
diff --git a/Engine/Plugins/VirtualProduction/VirtualScouting/Source/VirtualScoutingOpenXR/VirtualScoutingOpenXR.Build.cs b/Engine/Plugins/VirtualProduction/VirtualScouting/Source/VirtualScoutingOpenXR/VirtualScoutingOpenXR.Build.cs
new file mode 100644
--- /dev/null
+++ b/Engine/Plugins/VirtualProduction/VirtualScouting/Source/VirtualScoutingOpenXR/VirtualScoutingOpenXR.Build.cs
@@ -0,0 +1,33 @@
+// Copyright Epic Games, Inc. All Rights Reserved.
+
+using UnrealBuildTool;
+
+public class VirtualScoutingOpenXR : ModuleRules
+{
+	public VirtualScoutingOpenXR(ReadOnlyTargetRules Target) : base(Target)
+	{
+		PCHUsage = ModuleRules.PCHUsageMode.UseExplicitOrSharedPCHs;
+
+		PublicDependencyModuleNames.AddRange(new string[]
+			{
+				"Core",
+				"CoreUObject",
+				"Engine",
+				"InputCore",
+				"OpenXRHMD",
+			}
+		);
+
+		// This entire module is an Editor module in spirit, but in practice must be Runtime,
+		// because IOpenXRExtensionPlugin requires the PostConfigInit LoadingPhase.
+		if (Target.bBuildEditor)
+		{
+			PublicDependencyModuleNames.Add("VREditor");
+		}
+
+		PrivateDependencyModuleNames.AddRange(new string[]
+			{
+			}
+		);
+	}
+}
diff --git a/Engine/Plugins/VirtualProduction/VirtualScouting/VirtualScouting.uplugin b/Engine/Plugins/VirtualProduction/VirtualScouting/VirtualScouting.uplugin
new file mode 100644
--- /dev/null
+++ b/Engine/Plugins/VirtualProduction/VirtualScouting/VirtualScouting.uplugin
@@ -0,0 +1,113 @@
+{
+	"FileVersion": 3,
+	"Version": 1,
+	"VersionName": "1.0",
+	"FriendlyName": "Virtual Scouting",
+	"Description": "Virtual Scouting lets filmmakers scout a digital environment in virtual reality.",
+	"Category": "Virtual Production",
+	"CreatedBy": "Epic Games, Inc.",
+	"CreatedByURL": "https://epicgames.com",
+	"DocsURL": "",
+	"MarketplaceURL": "",
+	"SupportURL": "",
+	"CanContainContent": true,
+	"IsBetaVersion": false,
+	"IsExperimentalVersion": false,
+	"Installed": false,
+	"EnabledByDefault": false,
+	"SupportedTargetPlatforms": [
+		"Win64"
+	],
+	"Modules": [
+		{
+			"Name": "VirtualScoutingOpenXR",
+			"Type": "Runtime",
+			"LoadingPhase": "PostConfigInit",
+			"PlatformAllowList": [
+				"Win64"
+			]
+		},
+		{
+			"Name": "VirtualScouting",
+			"Type": "Runtime",
+			"LoadingPhase": "Default",
+			"PlatformAllowList": [
+				"Win64"
+			]
+		},
+		{
+			"Name": "VirtualScoutingEditor",
+			"Type": "Editor",
+			"LoadingPhase": "Default",
+			"PlatformAllowList": [
+				"Win64"
+			]
+		}
+	],
+	"Plugins": [
+		{
+			"Name": "XRCreativeFramework",
+			"Enabled":  true
+		},
+		{
+			"Name": "OpenXR",
+			"Enabled": true
+		},
+		{
+			"Name": "EnhancedInput",
+			"Enabled": true
+		},
+		{
+			"Name": "CommonUI",
+			"Enabled": true
+		},
+		{
+			"Name": "BlueprintFileUtils",
+			"Enabled": true
+		},
+		{
+			"Name": "JsonBlueprintUtilities",
+			"Enabled": true
+		},
+		{
+			"Name": "MultiUserClient",
+			"Enabled": true
+		},
+		{
+			"Name": "LiveLink",
+			"Enabled": true
+		},
+		{
+			"Name": "GeometryScripting",
+			"Enabled": true
+		},
+		{
+			"Name": "GizmoFramework",
+			"Enabled": true
+		},
+		{
+			"Name": "GizmoEdMode",
+			"Enabled": true
+		},
+		{
+			"Name": "ModelingToolsEditorMode",
+			"Enabled": true
+		},
+		{
+			"Name": "ConsoleVariables",
+			"Enabled": true
+		},
+		{
+			"Name": "LevelSnapshots",
+			"Enabled": true
+		},
+		{
+			"Name": "CineCameraSceneCapture",
+			"Enabled": true
+		},
+		{
+			"Name": "VirtualProductionUtilities",
+			"Enabled": true
+		}
+	]
+}
