Virtual Scouting polish

- tidied up naming on placment tool assets
- renamed Virtual Scouting settings property for placement tool content collection

#rnx
#virtualized

diff --git a/Engine/Plugins/VirtualProduction/VirtualScouting/Source/VirtualScouting/Public/VirtualScoutingSettings.h b/Engine/Plugins/VirtualProduction/VirtualScouting/Source/VirtualScouting/Public/VirtualScoutingSettings.h
--- a/Engine/Plugins/VirtualProduction/VirtualScouting/Source/VirtualScouting/Public/VirtualScoutingSettings.h
+++ b/Engine/Plugins/VirtualProduction/VirtualScouting/Source/VirtualScouting/Public/VirtualScoutingSettings.h
@@ -43,8 +43,8 @@ public:
 	UPROPERTY(Config, EditAnywhere, BlueprintReadWrite, Category="Virtual Scouting", meta=(DisplayName="Sequence Tool Collection"))
 	FName SequenceToolCollection;
 
-	UPROPERTY(Config, EditAnywhere, BlueprintReadWrite, Category = "Virtual Scouting", meta = (DisplayName = "Content Browser Tool Collection"))
-	FName ContentBrowserToolCollection;
+	UPROPERTY(Config, EditAnywhere, BlueprintReadWrite, Category = "Virtual Scouting", meta = (DisplayName = "Placement Tool Collection"))
+	FName PlacementToolCollection;
 };
 
 /**
