[Gizmos] Remove GizmoEdMode from areas not covered by preflight checks

#rnx
#Synced-CL 51719765

diff --git a/Engine/Plugins/VirtualProduction/VirtualScouting/VirtualScouting.uplugin b/Engine/Plugins/VirtualProduction/VirtualScouting/VirtualScouting.uplugin
--- a/Engine/Plugins/VirtualProduction/VirtualScouting/VirtualScouting.uplugin
+++ b/Engine/Plugins/VirtualProduction/VirtualScouting/VirtualScouting.uplugin
@@ -85,10 +85,6 @@
 			"Name": "GizmoFramework",
 			"Enabled": true
 		},
-		{
-			"Name": "GizmoEdMode",
-			"Enabled": true
-		},
 		{
 			"Name": "ModelingToolsEditorMode",
 			"Enabled": true
